Quixel Bifrost for UE5.

diff --git a/Engine/Plugins/Bridge/Bridge.uplugin b/Engine/Plugins/Bridge/Bridge.uplugin
new file mode 100644
--- /dev/null
+++ b/Engine/Plugins/Bridge/Bridge.uplugin
@@ -0,0 +1,32 @@
+{
+	"FileVersion": 3,
+	"Version": 1,
+	"VersionName": "0.10",
+	"FriendlyName": "Bridge",
+	"Description": "Megascans Link for Quixel Bridge.",
+	"Category": "Editor",
+	"CreatedBy": "Quixel",
+	"CreatedByURL": "http://www.quixel.se",
+	"DocsURL": "https://help.quixel.com/hc/en-us/sections/360000575958-UE4-Livelink-FAQs",
+	"MarketplaceURL": "",
+	"SupportURL": "[REDACTED_EMAIL]",	
+	"CanContainContent": true,
+	"EnabledByDefault" : true,
+	"SupportedTargetPlatforms": [
+		"Win64",
+		"Mac",
+		"Linux"
+	],
+	"Modules": [
+		{
+			"Name": "Bridge",
+			"Type": "Editor",
+			"LoadingPhase": "PostDefault",
+			"WhitelistPlatforms": [
+				"Win64",
+				"Mac",
+				"Linux"
+			]
+		}
+	]
+}
\ No newline at end of file
diff --git a/Engine/Plugins/Bridge/Config/FilterPlugin.ini b/Engine/Plugins/Bridge/Config/FilterPlugin.ini
new file mode 100644
--- /dev/null
+++ b/Engine/Plugins/Bridge/Config/FilterPlugin.ini
@@ -0,0 +1,8 @@
+[FilterPlugin]
+; This section lists additional files which will be packaged along with your plugin. Paths should be listed relative to the root plugin directory, and
+; may include "...", "*", and "?" wildcards to match directories, files, and individual characters respectively.
+;
+; Examples:
+;    /README.txt
+;    /Extras/...
+;    /Binaries/ThirdParty/*.dll
diff --git a/Engine/Plugins/Bridge/Source/Bridge/Bridge.Build.cs b/Engine/Plugins/Bridge/Source/Bridge/Bridge.Build.cs
new file mode 100644
--- /dev/null
+++ b/Engine/Plugins/Bridge/Source/Bridge/Bridge.Build.cs
@@ -0,0 +1,49 @@
+// Copyright Epic Games, Inc. All Rights Reserved.
+namespace UnrealBuildTool.Rules
+{
+	using System.Collections.Generic;
+	using UnrealBuildTool;
+	using System.IO;
+	using System;
+
+	public class Bridge  : ModuleRules
+	{
+		private string ThirdPartyPath
+		{
+			get { return Path.GetFullPath(Path.Combine(ModuleDirectory, "../../ThirdParty/")); }
+		}
+
+		private void stageFiles(String[] FilesToStage)
+		{
+			foreach (var File in FilesToStage)
+			{
+				RuntimeDependencies.Add(File);
+			}
+		}
+
+		public Bridge (ReadOnlyTargetRules Target) : base(Target)
+		{
+			PublicDependencyModuleNames.AddRange(
+				new string[] {
+					"Core",
+					"CoreUObject",
+					"Engine",
+					"InputCore",
+					"RenderCore",
+					"RHI",
+					"Slate",
+					"SlateCore",
+					"UMG",
+					"Json",
+					"WebBrowser",
+					"Networking",
+					"Sockets",
+					});
+
+			PrivateIncludePaths.AddRange(
+				new string[] {
+					Path.Combine(ModuleDirectory, "Private"),
+					});
+			}
+		}
+}
diff --git a/Engine/Plugins/Bridge/Source/Bridge/Private/BridgeModule.cpp b/Engine/Plugins/Bridge/Source/Bridge/Private/BridgeModule.cpp
new file mode 100644
--- /dev/null
+++ b/Engine/Plugins/Bridge/Source/Bridge/Private/BridgeModule.cpp
@@ -0,0 +1,30 @@
+// Copyright Epic Games, Inc. All Rights Reserved.
+#include "IBridgeModule.h"
+#include "UI/BridgeUIManager.h"
+#include "Misc/Paths.h"
+#include "CoreMinimal.h"
+#include "NodeProcess.h"
+
+#define LOCTEXT_NAMESPACE "Bridge"
+
+class FBridgeModule : public IBridgeModule
+{
+public:
+	virtual void StartupModule() override
+	{
+		if (GIsEditor && !IsRunningCommandlet())
+		{
+			 FBridgeUIManager::Initialize();
+		}
+		FNodeProcessManager::Get()->StartNodeProcess();
+	}
+
+	virtual void ShutdownModule() override
+	{
+		FNodeProcessManager::Get()->KillNodeProcess();
+	}
+};
+
+IMPLEMENT_MODULE(FBridgeModule, Bridge);
+
+#undef LOCTEXT_NAMESPACE
diff --git a/Engine/Plugins/Bridge/Source/Bridge/Private/NodeCommManager.cpp b/Engine/Plugins/Bridge/Source/Bridge/Private/NodeCommManager.cpp
new file mode 100644
--- /dev/null
+++ b/Engine/Plugins/Bridge/Source/Bridge/Private/NodeCommManager.cpp
@@ -0,0 +1,18 @@
+// Copyright Epic Games, Inc. All Rights Reserved.
+#include "NodeCommManager.h"
+
+TSharedPtr<FNodeCommManager> FNodeCommManager::NodeCommManager;
+
+TSharedPtr<FNodeCommManager> FNodeCommManager::Get()
+{
+	if (!NodeCommManager.IsValid())
+	{
+		NodeCommManager = MakeShareable(new FNodeCommManager);
+	}
+	return NodeCommManager;
+}
+
+void FNodeCommManager::NodeDataReceived(const FString& NodeData)
+{
+	UE_LOG(LogTemp, Error, TEXT("Data received from Node process"));
+}
diff --git a/Engine/Plugins/Bridge/Source/Bridge/Private/NodeCommManager.h b/Engine/Plugins/Bridge/Source/Bridge/Private/NodeCommManager.h
new file mode 100644
--- /dev/null
+++ b/Engine/Plugins/Bridge/Source/Bridge/Private/NodeCommManager.h
@@ -0,0 +1,29 @@
+// Copyright Epic Games, Inc. All Rights Reserved.
+#pragma once
+#include "CoreMinimal.h"
+#include "NodeCommManager.generated.h"
+
+class FNodeCommManager
+{
+private:
+	FNodeCommManager() = default;
+	static TSharedPtr<FNodeCommManager> NodeCommManager;	
+
+public:
+	static TSharedPtr<FNodeCommManager> Get();
+	void NodeDataReceived(const FString& NodeData);
+	
+};
+
+//Dummy data structure, all json data to be converted to ustructs
+USTRUCT()
+struct FNodeSpawnInfo
+{
+	GENERATED_USTRUCT_BODY();
+public:
+	UPROPERTY(EditAnywhere,Category = "NodeData")
+		FString Name;
+
+	UPROPERTY(EditAnywhere, Category = "NodeData")
+		FString Value;
+};
\ No newline at end of file
diff --git a/Engine/Plugins/Bridge/Source/Bridge/Private/NodePort.cpp b/Engine/Plugins/Bridge/Source/Bridge/Private/NodePort.cpp
new file mode 100644
--- /dev/null
+++ b/Engine/Plugins/Bridge/Source/Bridge/Private/NodePort.cpp
@@ -0,0 +1,26 @@
+// Copyright Epic Games, Inc. All Rights Reserved.
+#include "NodePort.h"
+#include "UI/BridgeUIManager.h"
+#include "Misc/Paths.h"
+#include "Misc/FileHelper.h"
+
+UNodePort::UNodePort(const FObjectInitializer& ObjectInitializer)
+	: Super(ObjectInitializer)
+{
+}
+
+FString UNodePort::GetNodePort()
+{
+	FString NodePort;
+	if (FFileHelper::LoadFileToString(NodePort, *PortFilePath))
+	{
+		return NodePort;
+	}
+
+	return TEXT("0");
+}
+
+bool UNodePort::IsNodeRunning()
+{
+	return true;
+}
diff --git a/Engine/Plugins/Bridge/Source/Bridge/Private/NodePort.h b/Engine/Plugins/Bridge/Source/Bridge/Private/NodePort.h
new file mode 100644
--- /dev/null
+++ b/Engine/Plugins/Bridge/Source/Bridge/Private/NodePort.h
@@ -0,0 +1,19 @@
+// Copyright Epic Games, Inc. All Rights Reserved.
+#pragma once
+#include "CoreMinimal.h"
+#include "NodePort.generated.h"
+
+UCLASS()
+class UNodePort :
+	public UObject
+{
+	GENERATED_UCLASS_BODY()
+public:
+	FString PortFilePath = FPaths::Combine(FPaths::EnginePluginsDir(), TEXT("Bridge"), TEXT("Content"), TEXT("node_port.txt"));
+	bool bIsNodeRunning;
+	UFUNCTION()
+		FString GetNodePort();
+
+	UFUNCTION()
+		bool IsNodeRunning();
+};
diff --git a/Engine/Plugins/Bridge/Source/Bridge/Private/NodeProcess.cpp b/Engine/Plugins/Bridge/Source/Bridge/Private/NodeProcess.cpp
new file mode 100644
--- /dev/null
+++ b/Engine/Plugins/Bridge/Source/Bridge/Private/NodeProcess.cpp
@@ -0,0 +1,97 @@
+// Copyright Epic Games, Inc. All Rights Reserved.
+#include "NodeProcess.h"
+#include "Misc/Paths.h"
+
+#include <iostream>
+#include <stdlib.h>
+#include <string>
+
+TSharedPtr<FNodeProcessManager> FNodeProcessManager::NodeProcessManager;
+
+TSharedPtr<FNodeProcessManager> FNodeProcessManager::Get()
+{
+	if (!NodeProcessManager.IsValid())
+	{
+		NodeProcessManager = MakeShareable(new FNodeProcessManager);
+	}
+	return NodeProcessManager;
+}
+
+const FString FNodeProcessManager::GetPluginPath() const
+{
+	return FPaths::Combine(FPaths::EnginePluginsDir(), BridgePluginName);
+}
+
+const FString FNodeProcessManager::GetProcessURL() const
+{
+#if PLATFORM_WINDOWS
+	FString MainFilePath = FPaths::ConvertRelativePathToFull(FPaths::Combine(GetPluginPath(), TEXT("Content\\Win"), TEXT("index.exe")));
+	return *MainFilePath;
+#else 
+	FString MainFilePath = FPaths::ConvertRelativePathToFull(FPaths::Combine(GetPluginPath(), TEXT("Content\\Mac"), TEXT("index")));
+	return *MainFilePath;
+#endif
+}
+
+void FNodeProcessManager::StartNodeProcess()
+{
+	const FString ProcessURL = GetProcessURL();
+
+#if PLATFORM_WINDOWS
+	NodeProcessHandle = FPlatformProcess::CreateProc(*ProcessURL, NULL, true, true, true, &OutProcessId, 0, NULL, NULL, NULL);
+	if (!NodeProcessHandle.IsValid())
+	{
+		UE_LOG(LogTemp, Error, TEXT("There was a problem in creating the Node process"));
+	}
+	int32 ReturnCode;	
+	FPlatformProcess::GetProcReturnCode(NodeProcessHandle, &ReturnCode);
+	if (ReturnCode < 0)
+	{
+		UE_LOG(LogTemp, Error, TEXT("Error : Process return %s"), *FString::FromInt(ReturnCode));
+	}
+
+	/* std::string command = "echo " + std::to_string(OutProcessId) + " >> WIN_NODE_PROCESS_ID.txt"; */
+	/* system(command.c_str()); */
+#elif PLATFORM_MAC
+	// Set env variable for logged in user
+ 	const TCHAR *loginVar = TEXT("LOGIN");
+	const FString userName = FMacPlatformMisc::GetEnvironmentVariable(loginVar);
+	
+	const FString path = "[REDACTED_USER_HOME] + userName + "/bin:/usr/local/bin:/usr/local/bin:/usr/bin:/bin:/usr/sbin:/sbni";
+	
+	const TCHAR *pathVar = TEXT("PATH");
+	const TCHAR *pathValue = *path;
+	FMacPlatformMisc::SetEnvironmentVar(pathVar, pathValue);
+
+	// Start node process
+	NodeProcessHandle = FPlatformProcess::CreateProc(*ProcessURL, NULL, true, true, true, &OutProcessId, 0, NULL, NULL, NULL); 
+
+	if (!NodeProcessHandle.IsValid())
+	{
+		/* std::string error_command = "echo " + std::to_string(102) + " >> MAC_NODE_PROCESS_ID.txt"; */
+		/* system(error_command.c_str()); */
+		UE_LOG(LogTemp, Error, TEXT("There was a problem in creating the Node process"));
+	}
+	int32 ReturnCode;	
+	FPlatformProcess::GetProcReturnCode(NodeProcessHandle, &ReturnCode);
+	if (ReturnCode < 0)
+	{
+		UE_LOG(LogTemp, Error, TEXT("Error : Process return %s"), *FString::FromInt(ReturnCode));
+	}
+
+	/* std::string mac_command = "echo " + std::to_string(OutProcessId) + " >> MAC_NODE_PROCESS_ID.txt"; */
+	/* system(mac_command.c_str()); */
+#endif
+}
+
+void FNodeProcessManager::KillNodeProcess()
+{
+	if (OutProcessId > 0) {
+		FPlatformProcess::TerminateProc(NodeProcessHandle, true);
+	}
+}
+
+void FNodeProcessManager::HandleBrowserUrlChanged(const FText& Url)
+{
+	UE_LOG(LogTemp, Error, TEXT("URL changed"));
+}
diff --git a/Engine/Plugins/Bridge/Source/Bridge/Private/NodeProcess.h b/Engine/Plugins/Bridge/Source/Bridge/Private/NodeProcess.h
new file mode 100644
--- /dev/null
+++ b/Engine/Plugins/Bridge/Source/Bridge/Private/NodeProcess.h
@@ -0,0 +1,29 @@
+// Copyright Epic Games, Inc. All Rights Reserved.
+#pragma once
+
+#if PLATFORM_WINDOWS
+#include "GenericPlatform/GenericPlatformProcess.h"
+#elif PLATFORM_MAC
+#include "Mac/MacPlatformMisc.h"
+#endif
+
+class FNodeProcessManager
+{
+private:
+	FNodeProcessManager() = default;	
+	static TSharedPtr<FNodeProcessManager> NodeProcessManager;
+
+	const FString BridgePluginName = TEXT("Bridge");
+	uint32 OutProcessId = 0;
+	FProcHandle NodeProcessHandle;
+
+	const FString GetProcessURL() const;
+	const FString GetPluginPath() const;	
+
+public:
+	static TSharedPtr<FNodeProcessManager> Get();
+	void StartNodeProcess();
+	void KillNodeProcess() ;
+	bool bIsNodeRunning = false ;
+	void HandleBrowserUrlChanged(const FText& Url);
+};
diff --git a/Engine/Plugins/Bridge/Source/Bridge/Private/TCPServer.cpp b/Engine/Plugins/Bridge/Source/Bridge/Private/TCPServer.cpp
new file mode 100644
--- /dev/null
+++ b/Engine/Plugins/Bridge/Source/Bridge/Private/TCPServer.cpp
@@ -0,0 +1,181 @@
+// Copyright Epic Games, Inc. All Rights Reserved.
+#include "TCPServer.h"
+#include <string>
+#include "NetworkMessage.h"
+#include "Async/Async.h"
+#include "Misc/ScopedSlowTask.h"
+
+
+FTCPServer::FTCPServer() 
+{
+	FThreadSafeCounter  WorkerCounter;
+	FString ThreadName(FString::Printf(TEXT("MegascansPlugin%i"), WorkerCounter.Increment()));
+	ClientThread = FRunnableThread::Create(this, *ThreadName, 8 * 1024, TPri_Normal);
+	
+}
+
+
+FTCPServer::~FTCPServer()
+{
+	
+	Stop();
+
+	
+	if (Listener != NULL)
+	{
+		Listener->Stop();
+		delete Listener;
+		Listener = NULL;
+	}
+	
+	if (!PendingClients.IsEmpty())
+	{
+		FSocket *Client = NULL;
+
+		while (PendingClients.Dequeue(Client))
+		{
+			Client->Close();
+			delete Client;
+			Client = NULL;
+		}
+	}
+
+	for (TArray<class FSocket*>::TIterator ClientIt(Clients); ClientIt; ++ClientIt)
+	{
+		(*ClientIt)->Close();
+		delete (*ClientIt);
+		(*ClientIt) = NULL;
+	}
+	
+	if (ClientThread != NULL)
+	{
+		ClientThread->Kill(true);
+		delete ClientThread;
+	}
+}
+
+
+bool FTCPServer::Init()
+{	
+
+	if (Listener == NULL)
+	{
+		FIPv4Address address;
+		FIPv4Address::Parse(LocalHostIP, address);
+		FIPv4Endpoint endPoint = FIPv4Endpoint(address, PortNum);
+		Listener = new FTcpListener(endPoint, FTimespan::FromMilliseconds(300));
+		Listener->OnConnectionAccepted().BindRaw(this, &FTCPServer::HandleListenerConnectionAccepted);
+		Stopping = false;
+	}
+	
+	return (Listener != NULL);
+}
+
+
+uint32  FTCPServer::Run()
+{
+	while (!Stopping)
+	{
+
+		if (!PendingClients.IsEmpty())
+		{
+			
+			for (TArray<class FSocket*>::TIterator ClientIt(Clients); ClientIt; ++ClientIt)
+			{
+				(*ClientIt)->Close();
+				delete (*ClientIt);
+				(*ClientIt) = NULL;
+			}
+
+			FSocket *Client = NULL;
+			while (PendingClients.Dequeue(Client))
+			{
+				
+				Clients.Empty();
+				Clients.Add(Client);
+				ConnectionTimer.AddZeroed(1);
+			}
+		}
+
+		
+
+		if (Clients.Num() > 0)
+		{
+			
+			FString RecievedJson = TEXT("");
+			uint32 DataSize = 0;
+
+			for (TArray<class FSocket*>::TIterator ClientIt(Clients); ClientIt; ++ClientIt)
+			{
+				
+				FSocket *Client = *ClientIt;				
+				DataSize = 0;
+				FString Request;
+				TArray<uint8> data;
+				bool haveMessage = false;
+				while (Client->HasPendingData(DataSize))
+				{					
+					haveMessage = RecvMessage(Client, DataSize, Request);										
+					if (haveMessage)
+					{					
+
+						int32 timer = 0;						
+						RecievedJson += Request;						
+						Request.Empty();
+						haveMessage = false;
+					}
+
+				}		
+
+			}
+
+			if (RecievedJson != TEXT(""))
+			{
+				if (!IsGarbageCollecting() && !GIsSavingPackage)
+				{			
+
+					
+					RecievedJson.Empty();
+				}
+			}
+			
+		}
+		
+		
+		FPlatformProcess::Sleep(0.3f); //0.3f
+	}
+
+	return 0;
+}
+
+
+bool FTCPServer::RecvMessage(FSocket *Socket, uint32 DataSize, FString& Message)
+{
+	check(Socket);	
+	FArrayReaderPtr Datagram = MakeShareable(new FArrayReader(true));
+	int32 stuff = 16;
+	Datagram->Init(FMath::Min(DataSize, 65507u), 81920);
+	int32 BytesRead = 0;
+
+	if (Socket->Recv(Datagram->GetData(), Datagram->Num(), BytesRead))
+	{
+		
+		char* Data = (char*)Datagram->GetData();
+		Data[BytesRead] = '\0';
+		FString message = UTF8_TO_TCHAR(Data);
+		Message = message;
+
+		return true;
+	}
+
+	return false;
+}
+
+bool FTCPServer::HandleListenerConnectionAccepted(class FSocket *ClientSocket, const FIPv4Endpoint& ClientEndpoint)
+
+{
+	PendingClients.Enqueue(ClientSocket);
+	return true;
+}
+
+
diff --git a/Engine/Plugins/Bridge/Source/Bridge/Private/TCPServer.h b/Engine/Plugins/Bridge/Source/Bridge/Private/TCPServer.h
new file mode 100644
--- /dev/null
+++ b/Engine/Plugins/Bridge/Source/Bridge/Private/TCPServer.h
@@ -0,0 +1,54 @@
+// Copyright Epic Games, Inc. All Rights Reserved.
+#pragma once
+
+
+#include "CoreMinimal.h"
+#include <string>
+#include "Sockets.h"
+#include "Misc/DateTime.h"
+#include "Async/AsyncWork.h"
+#include "Sockets.h"
+#include "Common/TcpListener.h"
+#include "Sockets.h"
+#include "Containers/Queue.h"
+#include "Common/UdpSocketReceiver.h"
+
+
+
+class FTCPServer : public FRunnable
+{
+
+public:	
+	FTCPServer();
+	~FTCPServer();	
+	virtual bool Init() override;	
+	virtual uint32 Run() override;
+
+
+
+	virtual void Stop() override
+	{
+		Stopping = true;
+	}
+
+	virtual void Exit() override { }
+	bool RecvMessage(FSocket *Socket, uint32 DataSize, FString& Message);
+	bool HandleListenerConnectionAccepted(class FSocket *ClientSocket, const FIPv4Endpoint& ClientEndpoint);
+
+
+	FSocket* ListenerSocket;
+	FString LocalHostIP = "127.0.0.1";
+	int32 PortNum = 13429; 
+	int32 ConnectionTimeout;
+	TArray<class FSocket*> Clients;
+
+private:	
+	TQueue<class FSocket*, EQueueMode::Mpsc> PendingClients;
+	bool Stopping;
+	FRunnableThread* ClientThread = NULL;
+	class FTcpListener *Listener = NULL;	
+	TArray<int32> ConnectionTimer;
+
+};
+
+
diff --git a/Engine/Plugins/Bridge/Source/Bridge/Private/UI/BridgeStyle.cpp b/Engine/Plugins/Bridge/Source/Bridge/Private/UI/BridgeStyle.cpp
new file mode 100644
--- /dev/null
+++ b/Engine/Plugins/Bridge/Source/Bridge/Private/UI/BridgeStyle.cpp
@@ -0,0 +1,72 @@
+// Copyright Epic Games, Inc. All Rights Reserved.
+#include "UI/BridgeStyle.h"
+
+#include "Framework/Application/SlateApplication.h"
+#include "Styling/SlateStyle.h"
+#include "Styling/SlateStyleRegistry.h"
+
+TUniquePtr<FSlateStyleSet> FBridgeStyle::MSStyleInstance;
+
+void FBridgeStyle::Initialize()
+{
+	if (!MSStyleInstance.IsValid())
+	{
+		MSStyleInstance = Create();
+		FSlateStyleRegistry::RegisterSlateStyle(*MSStyleInstance);
+	}
+}
+
+void FBridgeStyle::Shutdown()
+{
+	if (MSStyleInstance.IsValid())
+	{
+		
+		FSlateStyleRegistry::UnRegisterSlateStyle(*MSStyleInstance);
+		MSStyleInstance.Reset();
+	}
+}
+
+FName FBridgeStyle::GetStyleSetName()
+{
+	static FName StyleSetName(TEXT("BridgeStyle"));
+	return StyleSetName;
+}
+
+FName FBridgeStyle::GetContextName()
+{
+	//static FName ContextName(TEXT("Megascans"));
+	return FName(TEXT("Bridge"));
+}
+
+#define IMAGE_BRUSH( RelativePath, ... ) FSlateImageBrush( Style->RootToContentDir( RelativePath, TEXT(".png") ), __VA_ARGS__ )
+const FVector2D Icon20x20(20.0f, 20.0f);
+const FVector2D Icon40x40(40.0f, 40.0f);
+
+TUniquePtr< FSlateStyleSet > FBridgeStyle::Create()
+{
+	TUniquePtr< FSlateStyleSet > Style = MakeUnique<FSlateStyleSet>(GetStyleSetName());
+	Style->SetContentRoot(FPaths::EnginePluginsDir() / TEXT("Bridge/Resources"));
+	return Style;
+}
+
+void FBridgeStyle::SetIcon(const FString& StyleName, const FString& ResourcePath)
+{
+	FSlateStyleSet* Style = MSStyleInstance.Get();
+
+	FString Name(GetContextName().ToString());
+	Name = Name + "." + StyleName;
+	Style->Set(*Name, new IMAGE_BRUSH(ResourcePath, Icon40x40));
+
+	Name += ".Small";
+	Style->Set(*Name, new IMAGE_BRUSH(ResourcePath, Icon20x20));
+
+	FSlateApplication::Get().GetRenderer()->ReloadTextureResources();
+}
+
+#undef IMAGE_BRUSH
+
+const ISlateStyle& FBridgeStyle::Get()
+{
+	check(MSStyleInstance);
+	return *MSStyleInstance;
+}
diff --git a/Engine/Plugins/Bridge/Source/Bridge/Private/UI/BridgeStyle.h b/Engine/Plugins/Bridge/Source/Bridge/Private/UI/BridgeStyle.h
new file mode 100644
--- /dev/null
+++ b/Engine/Plugins/Bridge/Source/Bridge/Private/UI/BridgeStyle.h
@@ -0,0 +1,22 @@
+// Copyright Epic Games, Inc. All Rights Reserved.
+#pragma once
+#include "CoreMinimal.h"
+
+class FSlateStyleSet;
+class ISlateStyle;
+class FBridgeStyle
+{
+public:
+	static void Initialize();
+	static void Shutdown();
+
+	static const ISlateStyle& Get();
+	static FName GetStyleSetName();
+	static FName GetContextName();
+
+	static void SetIcon(const FString& StyleName, const FString& ResourcePath);
+
+private:
+	static TUniquePtr<FSlateStyleSet> Create();
+	static TUniquePtr<FSlateStyleSet> MSStyleInstance;
+};
diff --git a/Engine/Plugins/Bridge/Source/Bridge/Private/UI/BridgeUIManager.cpp b/Engine/Plugins/Bridge/Source/Bridge/Private/UI/BridgeUIManager.cpp
new file mode 100644
--- /dev/null
+++ b/Engine/Plugins/Bridge/Source/Bridge/Private/UI/BridgeUIManager.cpp
@@ -0,0 +1,180 @@
+// Copyright Epic Games, Inc. All Rights Reserved.
+#include "UI/BridgeUIManager.h"
+#include "UI/BridgeStyle.h"
+#include "LevelEditor.h"
+#include "Modules/ModuleManager.h"
+#include "NodePort.h"
+#include "NodeProcess.h"
+
+// WebBrowser
+#include "SWebBrowser.h"
+#include "WebBrowserModule.h"
+#include "IWebBrowserSingleton.h"
+#include "IWebBrowserCookieManager.h"
+// Widgets
+#include "Widgets/DeclarativeSyntaxSupport.h"
+#include "Widgets/SCompoundWidget.h"
+#include "Widgets/SBoxPanel.h"
+#include "Widgets/SWindow.h"
+#include "Widgets/Layout/SScrollBox.h"
+
+#define LOCTEXT_NAMESPACE "Bridge"
+#define LEVELEDITOR_MODULE_NAME TEXT("LevelEditor")
+#define CONTENTBROWSER_MODULE_NAME TEXT("ContentBrowser")
+
+TSharedPtr<FBridgeUIManagerImpl> FBridgeUIManager::Instance;
+UMegascansAuthentication* FBridgeUIManager::MegascansAuthentication;
+
+const FName BridgeTabName = "BridgeTab";
+
+class FBridgeUIManagerImpl
+{
+public:
+	void Initialize();
+	void Shutdown();
+	void HandleBrowserUrlChanged(const FText& Url);
+
+private:
+	void SetupMenuItem();
+	void CreateWIndow();
+	TSharedRef<SDockTab> CreateBridgeTab(const FSpawnTabArgs& Args);
+	void FillToolbar(FToolBarBuilder& ToolbarBuilder);
+	TSharedPtr<SWebBrowser> WebBrowserWidget;
+
+	void AddPluginMenu(FMenuBuilder& MenuBuilder);
+
+protected:
+	const FText BridgeTabDisplay = FText::FromString("Bridge");
+	const FText BridgeToolTip = FText::FromString("Launch Megascans Bridge");
+};
+
+void FBridgeUIManager::Initialize()
+{
+	if (!Instance.IsValid())
+	{
+		//Instance = MakeUnique<FBridgeUIManagerImpl>();
+		Instance = MakeShareable(new FBridgeUIManagerImpl);
+		Instance->Initialize();
+	}
+}
+
+void FBridgeUIManagerImpl::Initialize()
+{
+	FBridgeStyle::Initialize();
+	SetupMenuItem();
+}
+
+void FBridgeUIManagerImpl::SetupMenuItem()
+{
+	FBridgeStyle::SetIcon("Logo", "Logo80x80");
+	FBridgeStyle::SetIcon("ContextLogo", "Logo32x32");
+	FLevelEditorModule& LevelEditorModule = FModuleManager::LoadModuleChecked<FLevelEditorModule>(LEVELEDITOR_MODULE_NAME);
+	TSharedPtr<FExtender> ToolbarExtender = MakeShareable(new FExtender);
+	ToolbarExtender->AddToolBarExtension("Settings", EExtensionHook::After, nullptr, FToolBarExtensionDelegate::CreateRaw(this, &FBridgeUIManagerImpl::FillToolbar));
+	LevelEditorModule.GetToolBarExtensibilityManager()->AddExtender(ToolbarExtender);
+
+	// For Deleting Cookies
+	//IWebBrowserSingleton* WebBrowserSingleton = IWebBrowserModule::Get().GetSingleton();
+	//if (WebBrowserSingleton)
+	//{
+	//	TSharedPtr<IWebBrowserCookieManager> CookieManager = WebBrowserSingleton->GetCookieManager();
+	//	if (CookieManager.IsValid())
+	//	{
+	//		CookieManager->DeleteCookies();
+	//	}
+	//}
+
+	TSharedPtr<FExtender> NewMenuExtender = MakeShareable(new FExtender);
+	NewMenuExtender->AddMenuExtension("LevelEditor",
+		EExtensionHook::After,
+		NULL,
+		FMenuExtensionDelegate::CreateRaw(this, &FBridgeUIManagerImpl::AddPluginMenu));
+	LevelEditorModule.GetMenuExtensibilityManager()->AddExtender(NewMenuExtender);
+
+	FGlobalTabmanager::Get()->RegisterNomadTabSpawner(BridgeTabName,
+	FOnSpawnTab::CreateRaw(this, &FBridgeUIManagerImpl::CreateBridgeTab))
+		.SetDisplayName(BridgeTabDisplay)
+		.SetAutoGenerateMenuEntry(false)
+		.SetTooltipText(BridgeToolTip);
+}
+
+void FBridgeUIManagerImpl::AddPluginMenu(FMenuBuilder& MenuBuilder)
+{
+	MenuBuilder.BeginSection("CustomMenu", TAttribute<FText>(FText::FromString("Bifrost")));
+
+	MenuBuilder.AddMenuEntry(
+		LOCTEXT("OpenWindow", "Bifrost"),
+		LOCTEXT("ToolTip", "Open Quixel Bifrost"),
+		FSlateIcon(),
+		FUIAction(FExecuteAction::CreateRaw(this, &FBridgeUIManagerImpl::CreateWIndow))
+	);
+
+	MenuBuilder.EndSection();
+}
+
+void FBridgeUIManagerImpl::FillToolbar(FToolBarBuilder& ToolbarBuilder)
+{
+	ToolbarBuilder.BeginSection(TEXT("QuixelBridge"));
+	{
+		ToolbarBuilder.AddToolBarButton(
+			FUIAction(FExecuteAction::CreateRaw(this, &FBridgeUIManagerImpl::CreateWIndow)),
+			FName(TEXT("Quixel Bridge")),
+			LOCTEXT("QMSLiveLink_label", "Bridge"),
+			LOCTEXT("WorldProperties_ToolTipOverride", "Megascans Link with Bridge"),
+			FSlateIcon(FBridgeStyle::GetStyleSetName(), "Bridge.Logo"),
+			EUserInterfaceActionType::Button,
+			FName(TEXT("QuixelBridge"))
+		);
+	}
+	ToolbarBuilder.EndSection();
+}
+
+void FBridgeUIManagerImpl::HandleBrowserUrlChanged(const FText& Url)
+{
+	UE_LOG(LogTemp, Error, TEXT("URL changed"));
+}
+
+void FBridgeUIManagerImpl::CreateWIndow()
+{
+	FGlobalTabmanager::Get()->TryInvokeTab(BridgeTabName);
+}
+
+void FBridgeUIManager::Shutdown()
+{
+	FBridgeStyle::Shutdown();
+	FGlobalTabmanager::Get()->UnregisterNomadTabSpawner(BridgeTabName);
+}
+
+TSharedRef<SDockTab> FBridgeUIManagerImpl::CreateBridgeTab(const FSpawnTabArgs& Args)
+{
+	// Temp workaround which enables authentication (by impersonating Launcher's user agent)
+	FWebBrowserInitSettings browserInitSettings = FWebBrowserInitSettings();
+	browserInitSettings.ProductVersion = TEXT("EpicGamesLauncher/255.255.255-7654321+++Debug+Launcher UnrealEngine/4.23.0-0+UE4 Chrome/59.0.3071.15");
+	IWebBrowserModule::Get().CustomInitialize(browserInitSettings);
+
+	FString PluginPath = FPaths::Combine(FPaths::EnginePluginsDir(), TEXT("Bridge"));
+	FString IndexUrl = FPaths::ConvertRelativePathToFull(FPaths::Combine(PluginPath, TEXT("Content"), TEXT("out"), TEXT("index.html")));
+
+	TSharedPtr<SDockTab> BrowserDock;
+	{
+		SAssignNew(BrowserDock, SDockTab)
+		.TabRole(ETabRole::NomadTab)
+		[
+			SNew(SBorder)
+			.Padding(2)
+			[
+				SAssignNew(WebBrowserWidget, SWebBrowser)
+				.InitialURL(FPaths::Combine(TEXT("file:///"), IndexUrl))
+				.ShowControls(false)
+			]
+		];
+	}
+
+	UNodePort* NodePortInfo = NewObject<UNodePort>();
+	FBridgeUIManager::MegascansAuthentication = NewObject<UMegascansAuthentication>();
+	WebBrowserWidget->BindUObject(TEXT("NodePortInfo"), NodePortInfo, true);
+	WebBrowserWidget->BindUObject(TEXT("MegascansAuthentication"), FBridgeUIManager::MegascansAuthentication, true);
+	return BrowserDock.ToSharedRef();
+}
+
+#undef LOCTEXT_NAMESPACE
diff --git a/Engine/Plugins/Bridge/Source/Bridge/Private/UI/BridgeUIManager.h b/Engine/Plugins/Bridge/Source/Bridge/Private/UI/BridgeUIManager.h
new file mode 100644
--- /dev/null
+++ b/Engine/Plugins/Bridge/Source/Bridge/Private/UI/BridgeUIManager.h
@@ -0,0 +1,22 @@
+// Copyright Epic Games, Inc. All Rights Reserved.
+#pragma once
+
+#include "CoreMinimal.h"
+#include "Widgets/Docking/SDockTab.h"
+#include "SWebBrowser.h"
+#include "UI/SLoginWindow.h"
+
+class FBridgeUIManagerImpl;
+class FArguments;
+class SCompoundWidget;
+
+class  FBridgeUIManager
+{
+public:
+	static void Initialize();
+	static void Shutdown();
+	static TSharedPtr<FBridgeUIManagerImpl> Instance;
+	static UMegascansAuthentication* MegascansAuthentication;
+};
+
+
diff --git a/Engine/Plugins/Bridge/Source/Bridge/Private/UI/SLoginWindow.cpp b/Engine/Plugins/Bridge/Source/Bridge/Private/UI/SLoginWindow.cpp
new file mode 100644
--- /dev/null
+++ b/Engine/Plugins/Bridge/Source/Bridge/Private/UI/SLoginWindow.cpp
@@ -0,0 +1,64 @@
+// Copyright Epic Games, Inc. All Rights Reserved.
+#include "UI/SLoginWindow.h"
+#include "UI/BridgeUIManager.h"
+
+UMegascansAuthentication::UMegascansAuthentication(const FObjectInitializer& ObjectInitializer)
+	: Super(ObjectInitializer), LoginCode(TEXT(""))
+{
+}
+
+void UMegascansAuthentication::LoginCallback(FWebJSFunction LoginJSCallback)
+{
+	LoginCompleteDelegate.BindLambda(LoginJSCallback);
+}
+
+void UMegascansAuthentication::InitiateAuthentication()
+{
+	//Initialize a web browser for authentication
+	LoginMainWindow = SNew(SWindow)
+		.Title(FText::FromString(TEXT("Login")))
+		.ClientSize(FVector2D(450, 700))
+		.SupportsMaximize(false)
+		.SupportsMinimize(false)		
+		[
+			SNew(SVerticalBox)
+			+ SVerticalBox::Slot()
+			.HAlign(HAlign_Fill)
+			.VAlign(VAlign_Fill)
+			[
+				SNew(SMegascansLoginWindow)
+			]
+		];
+	FSlateApplication::Get().AddWindow(LoginMainWindow.ToSharedRef());
+}
+
+void SMegascansLoginWindow::HandleBrowserUrlChanged(const FText& Url)
+{
+	FString RedirectedUrl = Url.ToString();
+	const TCHAR *ProdCodeUrl = TEXT("https://quixel.com/?code=");
+	const TCHAR *StagingCodeUrl = TEXT("https://staging2.megascans.se/?code=");
+
+	const TCHAR *CodeUrl = StagingCodeUrl;
+
+	if (RedirectedUrl.StartsWith(CodeUrl))
+	{
+		FBridgeUIManager::MegascansAuthentication->LoginMainWindow->RequestDestroyWindow();
+		FString LoginCode = RedirectedUrl.Replace(CodeUrl, TEXT(""));
+		FBridgeUIManager::MegascansAuthentication->LoginCompleteDelegate.ExecuteIfBound(LoginCode);
+		UE_LOG(LogTemp, Error, TEXT("Found code : %s"), *RedirectedUrl.Replace(CodeUrl, TEXT("")));
+	}
+}
+
+void SMegascansLoginWindow::Construct(const FArguments& InArgs)
+{
+	FString Url = TEXT("https://www.quixel.com/login?return_to=https%3A%2F%2Fquixel.com%2Fmegascans%2Fhome");
+	FString StagingUrl = TEXT("https://staging2.megascans.se/login?return_to=https%3A%2F%2Fstaging2.megascans.se%2Fmegascans%2Fhome");
+
+	ChildSlot
+		[
+			SNew(SWebBrowser)
+			.InitialURL(StagingUrl)
+			.ShowControls(false)
+			.OnUrlChanged(this, &SMegascansLoginWindow::HandleBrowserUrlChanged)
+		];
+}
diff --git a/Engine/Plugins/Bridge/Source/Bridge/Private/UI/SLoginWindow.h b/Engine/Plugins/Bridge/Source/Bridge/Private/UI/SLoginWindow.h
new file mode 100644
--- /dev/null
+++ b/Engine/Plugins/Bridge/Source/Bridge/Private/UI/SLoginWindow.h
@@ -0,0 +1,38 @@
+// Copyright Epic Games, Inc. All Rights Reserved.
+#pragma once
+#include "CoreMinimal.h"
+#include "WebJSFunction.h"
+#include "SLoginWindow.generated.h"
+
+
+
+UCLASS()
+class UMegascansAuthentication :
+	public UObject
+{
+	GENERATED_UCLASS_BODY()
+		FString LoginCode;
+
+public:
+	DECLARE_DELEGATE_OneParam(FOnLoginCompleted, FString);
+	FOnLoginCompleted LoginCompleteDelegate;
+	TSharedPtr<SWindow> LoginMainWindow;
+
+	UFUNCTION()
+		void LoginCallback(FWebJSFunction LoginJSCallback);
+
+	UFUNCTION()
+		void InitiateAuthentication();
+};
+
+
+
+class SMegascansLoginWindow :
+	public SCompoundWidget
+{
+	SLATE_BEGIN_ARGS(SMegascansLoginWindow) {}
+	SLATE_END_ARGS()
+public:
+	void Construct(const FArguments& InArgs);
+	void HandleBrowserUrlChanged(const FText& Url);
+};
diff --git a/Engine/Plugins/Bridge/Source/Bridge/Public/IBridgeModule.h b/Engine/Plugins/Bridge/Source/Bridge/Public/IBridgeModule.h
new file mode 100644
--- /dev/null
+++ b/Engine/Plugins/Bridge/Source/Bridge/Public/IBridgeModule.h
@@ -0,0 +1,21 @@
+// Copyright Epic Games, Inc. All Rights Reserved.
+#pragma once
+
+#include "Modules/ModuleInterface.h"
+#include "Modules/ModuleManager.h"	
+
+#define MS_MODULE_NAME TEXT("Bridge")
+
+class IBridgeModule : public IModuleInterface
+{
+public:
+	static inline IBridgeModule& Get()
+	{
+		return FModuleManager::LoadModuleChecked<IBridgeModule>(MS_MODULE_NAME);
+	}
+
+	static inline bool IsAvailable()
+	{
+		return FModuleManager::Get().IsModuleLoaded(MS_MODULE_NAME);
+	}
+};
