Fix for warning : conflicting declarations of 'FNodeSpawnInfo'.

diff --git a/Engine/Plugins/Bridge/Source/Bridge/Private/NodeCommManager.h b/Engine/Plugins/Bridge/Source/Bridge/Private/NodeCommManager.h
--- a/Engine/Plugins/Bridge/Source/Bridge/Private/NodeCommManager.h
+++ b/Engine/Plugins/Bridge/Source/Bridge/Private/NodeCommManager.h
@@ -17,7 +17,7 @@ public:
 
 //Dummy data structure, all json data to be converted to ustructs
 USTRUCT()
-struct FNodeSpawnInfo
+struct FBifrostNodeInfo
 {
 	GENERATED_USTRUCT_BODY();
 public:
