Fixed start-up crash from passing null params to CreateProc

#rnx

diff --git a/Engine/Plugins/Bridge/Source/Bridge/Private/NodeProcess.cpp b/Engine/Plugins/Bridge/Source/Bridge/Private/NodeProcess.cpp
--- a/Engine/Plugins/Bridge/Source/Bridge/Private/NodeProcess.cpp
+++ b/Engine/Plugins/Bridge/Source/Bridge/Private/NodeProcess.cpp
@@ -38,7 +38,7 @@ void FNodeProcessManager::StartNodeProcess()
 	const FString ProcessURL = GetProcessURL();
 
 #if PLATFORM_WINDOWS
-	NodeProcessHandle = FPlatformProcess::CreateProc(*ProcessURL, NULL, true, true, true, &OutProcessId, 0, NULL, NULL, NULL);
+	NodeProcessHandle = FPlatformProcess::CreateProc(*ProcessURL, TEXT(""), true, true, true, &OutProcessId, 0, NULL, NULL, NULL);
 	if (!NodeProcessHandle.IsValid())
 	{
 		UE_LOG(LogTemp, Error, TEXT("There was a problem in creating the Node process"));
@@ -64,7 +64,7 @@ void FNodeProcessManager::StartNodeProcess()
 	FMacPlatformMisc::SetEnvironmentVar(pathVar, pathValue);
 
 	// Start node process
-	NodeProcessHandle = FPlatformProcess::CreateProc(*ProcessURL, NULL, true, true, true, &OutProcessId, 0, NULL, NULL, NULL); 
+	NodeProcessHandle = FPlatformProcess::CreateProc(*ProcessURL, TEXT(""), true, true, true, &OutProcessId, 0, NULL, NULL, NULL); 
 
 	if (!NodeProcessHandle.IsValid())
 	{
