Disabling Bifrost plugin by default

#rnx

diff --git a/Engine/Plugins/Bridge/Bridge.uplugin b/Engine/Plugins/Bridge/Bridge.uplugin
--- a/Engine/Plugins/Bridge/Bridge.uplugin
+++ b/Engine/Plugins/Bridge/Bridge.uplugin
@@ -11,7 +11,7 @@
 	"MarketplaceURL": "",
 	"SupportURL": "[REDACTED_EMAIL]",	
 	"CanContainContent": true,
-	"EnabledByDefault" : true,
+	"EnabledByDefault" : false,
 	"SupportedTargetPlatforms": [
 		"Win64",
 		"Mac",
