Non-Unity header fixes for Quixel Bridge.

diff --git a/Engine/Plugins/Bridge/Source/Bridge/Private/NodePort.cpp b/Engine/Plugins/Bridge/Source/Bridge/Private/NodePort.cpp
--- a/Engine/Plugins/Bridge/Source/Bridge/Private/NodePort.cpp
+++ b/Engine/Plugins/Bridge/Source/Bridge/Private/NodePort.cpp
@@ -1,8 +1,7 @@
 // Copyright Epic Games, Inc. All Rights Reserved.
 #include "NodePort.h"
 #include "UI/BridgeUIManager.h"
-#include "Misc/Paths.h"
-#include "Misc/FileHelper.h"
+
 
 UNodePort::UNodePort(const FObjectInitializer& ObjectInitializer)
 	: Super(ObjectInitializer)
diff --git a/Engine/Plugins/Bridge/Source/Bridge/Private/NodePort.h b/Engine/Plugins/Bridge/Source/Bridge/Private/NodePort.h
--- a/Engine/Plugins/Bridge/Source/Bridge/Private/NodePort.h
+++ b/Engine/Plugins/Bridge/Source/Bridge/Private/NodePort.h
@@ -1,6 +1,9 @@
 // Copyright Epic Games, Inc. All Rights Reserved.
 #pragma once
 #include "CoreMinimal.h"
+#include "Misc/Paths.h"
+#include "Misc/FileHelper.h"
+
 #include "NodePort.generated.h"
 
 UCLASS()
diff --git a/Engine/Plugins/Bridge/Source/Bridge/Private/NodeProcess.cpp b/Engine/Plugins/Bridge/Source/Bridge/Private/NodeProcess.cpp
--- a/Engine/Plugins/Bridge/Source/Bridge/Private/NodeProcess.cpp
+++ b/Engine/Plugins/Bridge/Source/Bridge/Private/NodeProcess.cpp
@@ -1,7 +1,8 @@
 // Copyright Epic Games, Inc. All Rights Reserved.
 #include "NodeProcess.h"
 #include "Misc/Paths.h"
-
+#include "Lumin/LuminPlatformProcess.h"
+#include "GenericPlatform/GenericPlatformMisc.h"
 #include <iostream>
 #include <stdlib.h>
 #include <string>
@@ -36,18 +37,19 @@ const FString FNodeProcessManager::GetProcessURL() const
 void FNodeProcessManager::StartNodeProcess()
 {
 	const FString ProcessURL = GetProcessURL();
+	FString ProcessParams = TEXT("");
 
 #if PLATFORM_WINDOWS
-	NodeProcessHandle = FPlatformProcess::CreateProc(*ProcessURL, TEXT(""), true, true, true, &OutProcessId, 0, NULL, NULL, NULL);
+	NodeProcessHandle = FPlatformProcess::CreateProc(*ProcessURL, *ProcessParams, true, true, true, &OutProcessId, 0, NULL, NULL, NULL);
 	if (!NodeProcessHandle.IsValid())
 	{
-		UE_LOG(LogTemp, Error, TEXT("There was a problem in creating the Node process"));
+		//UE_LOG(LogTemp, Error, TEXT("There was a problem in creating the Node process"));
 	}
 	int32 ReturnCode;	
 	FPlatformProcess::GetProcReturnCode(NodeProcessHandle, &ReturnCode);
 	if (ReturnCode < 0)
 	{
-		UE_LOG(LogTemp, Error, TEXT("Error : Process return %s"), *FString::FromInt(ReturnCode));
+		//UE_LOG(LogTemp, Error, TEXT("Error : Process return %s"), *FString::FromInt(ReturnCode));
 	}
 
 	/* std::string command = "echo " + std::to_string(OutProcessId) + " >> WIN_NODE_PROCESS_ID.txt"; */
@@ -70,13 +72,13 @@ void FNodeProcessManager::StartNodeProcess()
 	{
 		/* std::string error_command = "echo " + std::to_string(102) + " >> MAC_NODE_PROCESS_ID.txt"; */
 		/* system(error_command.c_str()); */
-		UE_LOG(LogTemp, Error, TEXT("There was a problem in creating the Node process"));
+		//UE_LOG(LogTemp, Error, TEXT("There was a problem in creating the Node process"));
 	}
 	int32 ReturnCode;	
 	FPlatformProcess::GetProcReturnCode(NodeProcessHandle, &ReturnCode);
 	if (ReturnCode < 0)
 	{
-		UE_LOG(LogTemp, Error, TEXT("Error : Process return %s"), *FString::FromInt(ReturnCode));
+		//UE_LOG(LogTemp, Error, TEXT("Error : Process return %s"), *FString::FromInt(ReturnCode));
 	}
 
 	/* std::string mac_command = "echo " + std::to_string(OutProcessId) + " >> MAC_NODE_PROCESS_ID.txt"; */
@@ -93,5 +95,5 @@ void FNodeProcessManager::KillNodeProcess()
 
 void FNodeProcessManager::HandleBrowserUrlChanged(const FText& Url)
 {
-	UE_LOG(LogTemp, Error, TEXT("URL changed"));
+	//UE_LOG(LogTemp, Error, TEXT("URL changed"));
 }
diff --git a/Engine/Plugins/Bridge/Source/Bridge/Private/NodeProcess.h b/Engine/Plugins/Bridge/Source/Bridge/Private/NodeProcess.h
--- a/Engine/Plugins/Bridge/Source/Bridge/Private/NodeProcess.h
+++ b/Engine/Plugins/Bridge/Source/Bridge/Private/NodeProcess.h
@@ -1,6 +1,7 @@
 // Copyright Epic Games, Inc. All Rights Reserved.
 #pragma once
-
+#include "CoreMinimal.h"
+#include "HoloLens/HoloLensPlatformProcess.h"
 #if PLATFORM_WINDOWS
 #include "GenericPlatform/GenericPlatformProcess.h"
 #elif PLATFORM_MAC
diff --git a/Engine/Plugins/Bridge/Source/Bridge/Private/TCPServer.cpp b/Engine/Plugins/Bridge/Source/Bridge/Private/TCPServer.cpp
--- a/Engine/Plugins/Bridge/Source/Bridge/Private/TCPServer.cpp
+++ b/Engine/Plugins/Bridge/Source/Bridge/Private/TCPServer.cpp
@@ -4,6 +4,7 @@
 #include "NetworkMessage.h"
 #include "Async/Async.h"
 #include "Misc/ScopedSlowTask.h"
+#include "UObject/GarbageCollection.h"
 
 
 FTCPServer::FTCPServer() 
@@ -131,7 +132,7 @@ uint32  FTCPServer::Run()
 
 			if (RecievedJson != TEXT(""))
 			{
-				if (!IsGarbageCollecting() && !GIsSavingPackage)
+				if (!IsGarbageCollecting() /*&& !GIsSavingPackage*/)
 				{			
 
 					
diff --git a/Engine/Plugins/Bridge/Source/Bridge/Private/UI/SLoginWindow.h b/Engine/Plugins/Bridge/Source/Bridge/Private/UI/SLoginWindow.h
--- a/Engine/Plugins/Bridge/Source/Bridge/Private/UI/SLoginWindow.h
+++ b/Engine/Plugins/Bridge/Source/Bridge/Private/UI/SLoginWindow.h
@@ -1,7 +1,11 @@
 // Copyright Epic Games, Inc. All Rights Reserved.
 #pragma once
 #include "CoreMinimal.h"
+#include "Widgets/SWindow.h"
+#include "Widgets/SCompoundWidget.h"
 #include "WebJSFunction.h"
+#include "Widgets/Input/SComboBox.h"
+#include "Internationalization/Text.h"
 #include "SLoginWindow.generated.h"
 
 
