Quixel Bridge compilation fixes.

diff --git a/Engine/Plugins/Bridge/Source/Bridge/Private/NodeProcess.cpp b/Engine/Plugins/Bridge/Source/Bridge/Private/NodeProcess.cpp
--- a/Engine/Plugins/Bridge/Source/Bridge/Private/NodeProcess.cpp
+++ b/Engine/Plugins/Bridge/Source/Bridge/Private/NodeProcess.cpp
@@ -1,8 +1,8 @@
 // Copyright Epic Games, Inc. All Rights Reserved.
 #include "NodeProcess.h"
 #include "Misc/Paths.h"
-#include "Lumin/LuminPlatformProcess.h"
-#include "GenericPlatform/GenericPlatformMisc.h"
+//#include "Lumin/LuminPlatformProcess.h"
+//#include "GenericPlatform/GenericPlatformMisc.h"
 #include <iostream>
 #include <stdlib.h>
 #include <string>
@@ -40,17 +40,17 @@ void FNodeProcessManager::StartNodeProcess()
 	FString ProcessParams = TEXT("");
 
 #if PLATFORM_WINDOWS
-	NodeProcessHandle = FPlatformProcess::CreateProc(*ProcessURL, *ProcessParams, true, true, true, &OutProcessId, 0, NULL, NULL, NULL);
-	if (!NodeProcessHandle.IsValid())
+	FPlatformProcess::CreateProc(*ProcessURL, *ProcessParams, true, true, true, &OutProcessId, 0, NULL, NULL, NULL);
+	/*if (!NodeProcessHandle.IsValid())
 	{
-		//UE_LOG(LogTemp, Error, TEXT("There was a problem in creating the Node process"));
+		
 	}
 	int32 ReturnCode;	
 	FPlatformProcess::GetProcReturnCode(NodeProcessHandle, &ReturnCode);
 	if (ReturnCode < 0)
 	{
-		//UE_LOG(LogTemp, Error, TEXT("Error : Process return %s"), *FString::FromInt(ReturnCode));
-	}
+		
+	}*/
 
 	/* std::string command = "echo " + std::to_string(OutProcessId) + " >> WIN_NODE_PROCESS_ID.txt"; */
 	/* system(command.c_str()); */
@@ -66,30 +66,27 @@ void FNodeProcessManager::StartNodeProcess()
 	FMacPlatformMisc::SetEnvironmentVar(pathVar, pathValue);
 
 	// Start node process
-	NodeProcessHandle = FPlatformProcess::CreateProc(*ProcessURL, TEXT(""), true, true, true, &OutProcessId, 0, NULL, NULL, NULL); 
+	FPlatformProcess::CreateProc(*ProcessURL, TEXT(""), true, true, true, &OutProcessId, 0, NULL, NULL, NULL); 
 
-	if (!NodeProcessHandle.IsValid())
+	/*if (!NodeProcessHandle.IsValid())
 	{
-		/* std::string error_command = "echo " + std::to_string(102) + " >> MAC_NODE_PROCESS_ID.txt"; */
-		/* system(error_command.c_str()); */
-		//UE_LOG(LogTemp, Error, TEXT("There was a problem in creating the Node process"));
+		
 	}
 	int32 ReturnCode;	
 	FPlatformProcess::GetProcReturnCode(NodeProcessHandle, &ReturnCode);
 	if (ReturnCode < 0)
 	{
-		//UE_LOG(LogTemp, Error, TEXT("Error : Process return %s"), *FString::FromInt(ReturnCode));
-	}
+		
+	}*/
+
 
-	/* std::string mac_command = "echo " + std::to_string(OutProcessId) + " >> MAC_NODE_PROCESS_ID.txt"; */
-	/* system(mac_command.c_str()); */
 #endif
 }
 
 void FNodeProcessManager::KillNodeProcess()
 {
 	if (OutProcessId > 0) {
-		FPlatformProcess::TerminateProc(NodeProcessHandle, true);
+		//FPlatformProcess::TerminateProc(NodeProcessHandle, true);
 	}
 }
 
diff --git a/Engine/Plugins/Bridge/Source/Bridge/Private/NodeProcess.h b/Engine/Plugins/Bridge/Source/Bridge/Private/NodeProcess.h
--- a/Engine/Plugins/Bridge/Source/Bridge/Private/NodeProcess.h
+++ b/Engine/Plugins/Bridge/Source/Bridge/Private/NodeProcess.h
@@ -1,7 +1,7 @@
 // Copyright Epic Games, Inc. All Rights Reserved.
 #pragma once
 #include "CoreMinimal.h"
-#include "HoloLens/HoloLensPlatformProcess.h"
+//#include "HoloLens/HoloLensPlatformProcess.h"
 #if PLATFORM_WINDOWS
 #include "GenericPlatform/GenericPlatformProcess.h"
 #elif PLATFORM_MAC
@@ -11,12 +11,12 @@
 class FNodeProcessManager
 {
 private:
-	FNodeProcessManager() = default;	
+	FNodeProcessManager() = default;
 	static TSharedPtr<FNodeProcessManager> NodeProcessManager;
 
 	const FString BridgePluginName = TEXT("Bridge");
 	uint32 OutProcessId = 0;
-	FProcHandle NodeProcessHandle;
+	//FProcHandle NodeProcessHandle;
 
 	const FString GetProcessURL() const;
 	const FString GetPluginPath() const;	
