Fixes for building with clang and UBT's -nopch option.

#trivial
#rnx

diff --git a/Engine/Plugins/Bridge/Source/Bridge/Private/UI/BridgeUIManager.cpp b/Engine/Plugins/Bridge/Source/Bridge/Private/UI/BridgeUIManager.cpp
--- a/Engine/Plugins/Bridge/Source/Bridge/Private/UI/BridgeUIManager.cpp
+++ b/Engine/Plugins/Bridge/Source/Bridge/Private/UI/BridgeUIManager.cpp
@@ -12,6 +12,7 @@
 #include "IWebBrowserSingleton.h"
 #include "IWebBrowserCookieManager.h"
 // Widgets
+#include "Framework/MultiBox/MultiBoxBuilder.h"
 #include "Widgets/DeclarativeSyntaxSupport.h"
 #include "Widgets/SCompoundWidget.h"
 #include "Widgets/SBoxPanel.h"
diff --git a/Engine/Source/Developer/TraceServices/Private/Model/AllocationsProvider.h b/Engine/Source/Developer/TraceServices/Private/Model/AllocationsProvider.h
--- a/Engine/Source/Developer/TraceServices/Private/Model/AllocationsProvider.h
+++ b/Engine/Source/Developer/TraceServices/Private/Model/AllocationsProvider.h
@@ -3,6 +3,8 @@
 #pragma once
 
 #include "TraceServices/Model/AllocationsProvider.h"
+#include "Containers/Map.h"
+#include "HAL/CriticalSection.h"
 
 namespace TraceServices
 {
diff --git a/Engine/Source/Developer/TraceServices/Private/Model/AllocationsQuery.h b/Engine/Source/Developer/TraceServices/Private/Model/AllocationsQuery.h
--- a/Engine/Source/Developer/TraceServices/Private/Model/AllocationsQuery.h
+++ b/Engine/Source/Developer/TraceServices/Private/Model/AllocationsQuery.h
@@ -4,6 +4,7 @@
 
 #include "AllocationsProvider.h"
 #include "Async/TaskGraphInterfaces.h"
+#include "Containers/Queue.h"
 #include "HAL/PlatformAtomics.h"
 
 #include <atomic>
diff --git a/Engine/Source/Developer/TraceServices/Private/Model/SbTree.cpp b/Engine/Source/Developer/TraceServices/Private/Model/SbTree.cpp
--- a/Engine/Source/Developer/TraceServices/Private/Model/SbTree.cpp
+++ b/Engine/Source/Developer/TraceServices/Private/Model/SbTree.cpp
@@ -1,6 +1,8 @@
 // Copyright Epic Games, Inc. All Rights Reserved.
 
 #include "SbTree.h"
+
+#include "Misc/FileHelper.h"
 #include "TraceServices/Containers/Allocators.h"
 
 #include <limits>
diff --git a/Engine/Source/Runtime/Core/Private/Experimental/Containers/HazardPointer.cpp b/Engine/Source/Runtime/Core/Private/Experimental/Containers/HazardPointer.cpp
--- a/Engine/Source/Runtime/Core/Private/Experimental/Containers/HazardPointer.cpp
+++ b/Engine/Source/Runtime/Core/Private/Experimental/Containers/HazardPointer.cpp
@@ -1,6 +1,7 @@
 // Copyright Epic Games, Inc. All Rights Reserved.
 
 #include "Experimental/Containers/HazardPointer.h"
+#include "Misc/ScopeLock.h"
 
 FHazardPointerCollection::~FHazardPointerCollection()
 {
diff --git a/Engine/Source/Runtime/Core/Public/Misc/DataDrivenPlatformInfoRegistry.h b/Engine/Source/Runtime/Core/Public/Misc/DataDrivenPlatformInfoRegistry.h
--- a/Engine/Source/Runtime/Core/Public/Misc/DataDrivenPlatformInfoRegistry.h
+++ b/Engine/Source/Runtime/Core/Public/Misc/DataDrivenPlatformInfoRegistry.h
@@ -6,6 +6,7 @@
 #include "Containers/UnrealString.h"
 #include "Containers/Map.h"
 #include "Containers/Array.h"
+#include "Misc/Guid.h"
 #include "UObject/NameTypes.h"
 #include "Internationalization/Text.h"
 
