Bifrost Non-Unity Compile fix.

diff --git a/Engine/Plugins/Bridge/Source/Bridge/Private/NodeProcess.cpp b/Engine/Plugins/Bridge/Source/Bridge/Private/NodeProcess.cpp
--- a/Engine/Plugins/Bridge/Source/Bridge/Private/NodeProcess.cpp
+++ b/Engine/Plugins/Bridge/Source/Bridge/Private/NodeProcess.cpp
@@ -39,8 +39,8 @@ void FNodeProcessManager::StartNodeProcess()
 	const FString ProcessURL = GetProcessURL();
 	FString ProcessParams = TEXT("");
 
-#if PLATFORM_WINDOWS
-	FPlatformProcess::CreateProc(*ProcessURL, *ProcessParams, true, true, true, &OutProcessId, 0, NULL, NULL, NULL);
+//#if PLATFORM_WINDOWS
+	//FPlatformProcess::CreateProc(*ProcessURL, *ProcessParams, true, true, true, &OutProcessId, 0, NULL, NULL, NULL);
 	/*if (!NodeProcessHandle.IsValid())
 	{
 		
@@ -54,7 +54,7 @@ void FNodeProcessManager::StartNodeProcess()
 
 	/* std::string command = "echo " + std::to_string(OutProcessId) + " >> WIN_NODE_PROCESS_ID.txt"; */
 	/* system(command.c_str()); */
-#elif PLATFORM_MAC
+//#elif PLATFORM_MAC
 	// Set env variable for logged in user
  	const TCHAR *loginVar = TEXT("LOGIN");
 	const FString userName = FMacPlatformMisc::GetEnvironmentVariable(loginVar);
@@ -63,10 +63,10 @@ void FNodeProcessManager::StartNodeProcess()
 	
 	const TCHAR *pathVar = TEXT("PATH");
 	const TCHAR *pathValue = *path;
-	FMacPlatformMisc::SetEnvironmentVar(pathVar, pathValue);
+	//FMacPlatformMisc::SetEnvironmentVar(pathVar, pathValue);
 
 	// Start node process
-	FPlatformProcess::CreateProc(*ProcessURL, TEXT(""), true, true, true, &OutProcessId, 0, NULL, NULL, NULL); 
+	//FPlatformProcess::CreateProc(*ProcessURL, TEXT(""), true, true, true, &OutProcessId, 0, NULL, NULL, NULL); 
 
 	/*if (!NodeProcessHandle.IsValid())
 	{
@@ -80,7 +80,7 @@ void FNodeProcessManager::StartNodeProcess()
 	}*/
 
 
-#endif
+//#endif
 }
 
 void FNodeProcessManager::KillNodeProcess()
