Compile fixes in NodeProcess.

diff --git a/Engine/Plugins/Bridge/Source/Bridge/Private/NodeProcess.cpp b/Engine/Plugins/Bridge/Source/Bridge/Private/NodeProcess.cpp
--- a/Engine/Plugins/Bridge/Source/Bridge/Private/NodeProcess.cpp
+++ b/Engine/Plugins/Bridge/Source/Bridge/Private/NodeProcess.cpp
@@ -56,13 +56,13 @@ void FNodeProcessManager::StartNodeProcess()
 	/* system(command.c_str()); */
 //#elif PLATFORM_MAC
 	// Set env variable for logged in user
- 	const TCHAR *loginVar = TEXT("LOGIN");
-	const FString userName = FMacPlatformMisc::GetEnvironmentVariable(loginVar);
+ 	//const TCHAR *loginVar = TEXT("LOGIN");
+	//const FString userName = FMacPlatformMisc::GetEnvironmentVariable(loginVar);
 	
-	const FString path = "[REDACTED_USER_HOME] + userName + "/bin:/usr/local/bin:/usr/local/bin:/usr/bin:/bin:/usr/sbin:/sbni";
+	//const FString path = "[REDACTED_USER_HOME] + userName + "/bin:/usr/local/bin:/usr/local/bin:/usr/bin:/bin:/usr/sbin:/sbni";
 	
-	const TCHAR *pathVar = TEXT("PATH");
-	const TCHAR *pathValue = *path;
+	//const TCHAR *pathVar = TEXT("PATH");
+	//const TCHAR *pathValue = *path;
 	//FMacPlatformMisc::SetEnvironmentVar(pathVar, pathValue);
 
 	// Start node process
