Quixel Bifrost Node process management.

diff --git a/Engine/Plugins/Bridge/Source/Bridge/Private/NodeProcess.cpp b/Engine/Plugins/Bridge/Source/Bridge/Private/NodeProcess.cpp
--- a/Engine/Plugins/Bridge/Source/Bridge/Private/NodeProcess.cpp
+++ b/Engine/Plugins/Bridge/Source/Bridge/Private/NodeProcess.cpp
@@ -1,12 +1,19 @@
 // Copyright Epic Games, Inc. All Rights Reserved.
 #include "NodeProcess.h"
 #include "Misc/Paths.h"
-//#include "Lumin/LuminPlatformProcess.h"
-//#include "GenericPlatform/GenericPlatformMisc.h"
+#include "GenericPlatform/GenericPlatformMisc.h"
+#include "GenericPlatform/GenericPlatformProcess.h"
 #include <iostream>
 #include <stdlib.h>
 #include <string>
 
+#if PLATFORM_WINDOWS
+#include "Windows/WindowsPlatformProcess.h"
+#elif PLATFORM_MAC
+#include "Mac/MacPlatformMisc.h"
+#include "Mac/MacPlatformProcess.h"
+#endif
+
 TSharedPtr<FNodeProcessManager> FNodeProcessManager::NodeProcessManager;
 
 TSharedPtr<FNodeProcessManager> FNodeProcessManager::Get()
@@ -39,7 +46,8 @@ void FNodeProcessManager::StartNodeProcess()
 	const FString ProcessURL = GetProcessURL();
 	FString ProcessParams = TEXT("");
 
-//#if PLATFORM_WINDOWS
+#if PLATFORM_WINDOWS
+	NodeProcessHandle = FPlatformProcess::CreateProc(*ProcessURL, *ProcessParams, true, true, true, &OutProcessId, 0, NULL, NULL, NULL);
 	//FPlatformProcess::CreateProc(*ProcessURL, *ProcessParams, true, true, true, &OutProcessId, 0, NULL, NULL, NULL);
 	/*if (!NodeProcessHandle.IsValid())
 	{
@@ -54,19 +62,19 @@ void FNodeProcessManager::StartNodeProcess()
 
 	/* std::string command = "echo " + std::to_string(OutProcessId) + " >> WIN_NODE_PROCESS_ID.txt"; */
 	/* system(command.c_str()); */
-//#elif PLATFORM_MAC
+#elif PLATFORM_MAC
 	// Set env variable for logged in user
- 	//const TCHAR *loginVar = TEXT("LOGIN");
-	//const FString userName = FMacPlatformMisc::GetEnvironmentVariable(loginVar);
+ 	const TCHAR *loginVar = TEXT("LOGIN");
+	const FString userName = FMacPlatformMisc::GetEnvironmentVariable(loginVar);
 	
-	//const FString path = "[REDACTED_USER_HOME] + userName + "/bin:/usr/local/bin:/usr/local/bin:/usr/bin:/bin:/usr/sbin:/sbni";
+	const FString path = "[REDACTED_USER_HOME] + userName + "/bin:/usr/local/bin:/usr/local/bin:/usr/bin:/bin:/usr/sbin:/sbni";
 	
-	//const TCHAR *pathVar = TEXT("PATH");
-	//const TCHAR *pathValue = *path;
-	//FMacPlatformMisc::SetEnvironmentVar(pathVar, pathValue);
+	const TCHAR *pathVar = TEXT("PATH");
+	const TCHAR *pathValue = *path;
+	FMacPlatformMisc::SetEnvironmentVar(pathVar, pathValue);
 
 	// Start node process
-	//FPlatformProcess::CreateProc(*ProcessURL, TEXT(""), true, true, true, &OutProcessId, 0, NULL, NULL, NULL); 
+	NodeProcessHandle = FPlatformProcess::CreateProc(*ProcessURL, TEXT(""), true, true, true, &OutProcessId, 0, NULL, NULL, NULL);
 
 	/*if (!NodeProcessHandle.IsValid())
 	{
@@ -80,13 +88,13 @@ void FNodeProcessManager::StartNodeProcess()
 	}*/
 
 
-//#endif
+#endif
 }
 
 void FNodeProcessManager::KillNodeProcess()
 {
 	if (OutProcessId > 0) {
-		//FPlatformProcess::TerminateProc(NodeProcessHandle, true);
+		FPlatformProcess::TerminateProc(NodeProcessHandle, true);
 	}
 }
 
diff --git a/Engine/Plugins/Bridge/Source/Bridge/Private/NodeProcess.h b/Engine/Plugins/Bridge/Source/Bridge/Private/NodeProcess.h
--- a/Engine/Plugins/Bridge/Source/Bridge/Private/NodeProcess.h
+++ b/Engine/Plugins/Bridge/Source/Bridge/Private/NodeProcess.h
@@ -1,12 +1,14 @@
 // Copyright Epic Games, Inc. All Rights Reserved.
 #pragma once
 #include "CoreMinimal.h"
-//#include "HoloLens/HoloLensPlatformProcess.h"
-#if PLATFORM_WINDOWS
-#include "GenericPlatform/GenericPlatformProcess.h"
-#elif PLATFORM_MAC
-#include "Mac/MacPlatformMisc.h"
-#endif
+
+//#if PLATFORM_WINDOWS
+//#include "GenericPlatform/GenericPlatformProcess.h"
+//#elif PLATFORM_MAC
+//#include "Mac/MacPlatformMisc.h"
+//#endif
+
+struct FProcHandle;
 
 class FNodeProcessManager
 {
@@ -16,7 +18,7 @@ private:
 
 	const FString BridgePluginName = TEXT("Bridge");
 	uint32 OutProcessId = 0;
-	//FProcHandle NodeProcessHandle;
+	FProcHandle NodeProcessHandle;
 
 	const FString GetProcessURL() const;
 	const FString GetPluginPath() const;	
