Bifrost Non-Unity compile fix.

diff --git a/Engine/Plugins/Bridge/Source/Bridge/Private/NodeProcess.cpp b/Engine/Plugins/Bridge/Source/Bridge/Private/NodeProcess.cpp
--- a/Engine/Plugins/Bridge/Source/Bridge/Private/NodeProcess.cpp
+++ b/Engine/Plugins/Bridge/Source/Bridge/Private/NodeProcess.cpp
@@ -1,18 +1,12 @@
 // Copyright Epic Games, Inc. All Rights Reserved.
 #include "NodeProcess.h"
 #include "Misc/Paths.h"
-#include "GenericPlatform/GenericPlatformMisc.h"
-#include "GenericPlatform/GenericPlatformProcess.h"
+
 #include <iostream>
 #include <stdlib.h>
 #include <string>
 
-#if PLATFORM_WINDOWS
-#include "Windows/WindowsPlatformProcess.h"
-#elif PLATFORM_MAC
-#include "Mac/MacPlatformMisc.h"
-#include "Mac/MacPlatformProcess.h"
-#endif
+
 
 TSharedPtr<FNodeProcessManager> FNodeProcessManager::NodeProcessManager;
 
diff --git a/Engine/Plugins/Bridge/Source/Bridge/Private/NodeProcess.h b/Engine/Plugins/Bridge/Source/Bridge/Private/NodeProcess.h
--- a/Engine/Plugins/Bridge/Source/Bridge/Private/NodeProcess.h
+++ b/Engine/Plugins/Bridge/Source/Bridge/Private/NodeProcess.h
@@ -1,14 +1,17 @@
 // Copyright Epic Games, Inc. All Rights Reserved.
 #pragma once
 #include "CoreMinimal.h"
+#include "GenericPlatform/GenericPlatformMisc.h"
+#include "GenericPlatform/GenericPlatformProcess.h"
 
-//#if PLATFORM_WINDOWS
-//#include "GenericPlatform/GenericPlatformProcess.h"
-//#elif PLATFORM_MAC
-//#include "Mac/MacPlatformMisc.h"
-//#endif
+#if PLATFORM_WINDOWS
+#include "Windows/WindowsPlatformProcess.h"
+#elif PLATFORM_MAC
+#include "Mac/MacPlatformMisc.h"
+#include "Mac/MacPlatformProcess.h"
+#endif
 
-struct FProcHandle;
+//struct FProcHandle;
 
 class FNodeProcessManager
 {
