Added UI entry points for Quixel Bifrost and updated icons.

diff --git a/Engine/Plugins/Bridge/Source/Bridge/Bridge.Build.cs b/Engine/Plugins/Bridge/Source/Bridge/Bridge.Build.cs
--- a/Engine/Plugins/Bridge/Source/Bridge/Bridge.Build.cs
+++ b/Engine/Plugins/Bridge/Source/Bridge/Bridge.Build.cs
@@ -38,6 +38,7 @@ namespace UnrealBuildTool.Rules
 					"WebBrowser",
 					"Networking",
 					"Sockets",
+					"ToolMenus",
 					});
 
 			PrivateIncludePaths.AddRange(
diff --git a/Engine/Plugins/Bridge/Source/Bridge/Private/UI/BridgeUIManager.cpp b/Engine/Plugins/Bridge/Source/Bridge/Private/UI/BridgeUIManager.cpp
--- a/Engine/Plugins/Bridge/Source/Bridge/Private/UI/BridgeUIManager.cpp
+++ b/Engine/Plugins/Bridge/Source/Bridge/Private/UI/BridgeUIManager.cpp
@@ -18,6 +18,8 @@
 #include "Widgets/SBoxPanel.h"
 #include "Widgets/SWindow.h"
 #include "Widgets/Layout/SScrollBox.h"
+#include "ToolMenus.h"
+#include "ToolMenuSection.h"
 
 #define LOCTEXT_NAMESPACE "Bridge"
 #define LEVELEDITOR_MODULE_NAME TEXT("LevelEditor")
@@ -85,6 +87,30 @@ void FBridgeUIManagerImpl::SetupMenuItem()
 	//	}
 	//}
 
+	// Adding Bridge entry to Quick Content menu.
+	UToolMenu* ContentMenu = UToolMenus::Get()->ExtendMenu("LevelEditor.LevelEditorToolBar.ContentQuickMenu");
+	FToolMenuSection& Section = ContentMenu->FindOrAddSection("GetContent");
+	Section.AddMenuEntry("OpenBridgeTab",
+		LOCTEXT("OpenBridgeTab_Label", "Quixel Bridge"),
+		LOCTEXT("OpenBridgeTab_Desc", "Opens the Quixel Bridge."),
+		FSlateIcon(FBridgeStyle::GetStyleSetName(), "Bridge.Logo"),
+		FUIAction(FExecuteAction::CreateRaw(this, &FBridgeUIManagerImpl::CreateWIndow), FCanExecuteAction())
+	);
+	Section.AddSeparator(NAME_None);
+
+	//Adding Bridge entry to Content Browser context and New menu.
+	UToolMenu* ContextMenu = UToolMenus::Get()->ExtendMenu("ContentBrowser.AddNewContextMenu");
+	//FToolMenuSection& ContextMenuSection = ContextMenu->AddSection("ContentBrowserMegascans", LOCTEXT("GetContentMenuHeading", "Quixel Content"));
+	FToolMenuSection& ContextMenuSection = ContextMenu->FindOrAddSection("ContentBrowserGetContent");
+	
+	ContextMenuSection.AddMenuEntry(
+		"GetMegascans",
+		LOCTEXT("OpenBridgeTabText", "Add Quixel Content"),
+		LOCTEXT("GetBridgeTooltip", "Add Megascans and DHI assets to project."),
+		FSlateIcon(FBridgeStyle::GetStyleSetName(), "Bridge.Logo"),
+		FUIAction(FExecuteAction::CreateRaw(this, &FBridgeUIManagerImpl::CreateWIndow), FCanExecuteAction())
+	);
+
 	TSharedPtr<FExtender> NewMenuExtender = MakeShareable(new FExtender);
 	NewMenuExtender->AddMenuExtension("LevelEditor",
 		EExtensionHook::After,
@@ -101,12 +127,12 @@ void FBridgeUIManagerImpl::SetupMenuItem()
 
 void FBridgeUIManagerImpl::AddPluginMenu(FMenuBuilder& MenuBuilder)
 {
-	MenuBuilder.BeginSection("CustomMenu", TAttribute<FText>(FText::FromString("Bifrost")));
+	MenuBuilder.BeginSection("CustomMenu", TAttribute<FText>(FText::FromString("Quixel")));
 
 	MenuBuilder.AddMenuEntry(
-		LOCTEXT("OpenWindow", "Bifrost"),
-		LOCTEXT("ToolTip", "Open Quixel Bifrost"),
-		FSlateIcon(),
+		LOCTEXT("OpenWindow", "Quixel Bridge"),
+		LOCTEXT("ToolTip", "Open Quixel Bridge"),
+		FSlateIcon(FBridgeStyle::GetStyleSetName(), "Bridge.Logo"),
 		FUIAction(FExecuteAction::CreateRaw(this, &FBridgeUIManagerImpl::CreateWIndow))
 	);
 
