Temp fix to Bridge.uplugin breaking Frosty Editor compilation and preflighting changes

diff --git a/Engine/Plugins/Bridge/Bridge.uplugin b/Engine/Plugins/Bridge/Bridge.uplugin
--- a/Engine/Plugins/Bridge/Bridge.uplugin
+++ b/Engine/Plugins/Bridge/Bridge.uplugin
@@ -26,27 +26,13 @@
 				"Win64",
 				"Mac"
 				
-			]
-		},
-		{
-			"Name": "MegascansPlugin",
-			"Type": "Editor",
-			"LoadingPhase": "PostDefault",
-			"WhitelistPlatforms": [
-				"Win64",
-				"Mac"
-				
 			]
 		}
-
 	],
 	"Plugins": [		
 		{
 			"Name": "EditorScriptingUtilities",
 			"Enabled": true
 		}
-
 	]
-
-
 }
\ No newline at end of file
