Quixel Bridge updated master materials and settings launch.

diff --git a/Engine/Plugins/Bridge/Source/Bridge/Bridge.Build.cs b/Engine/Plugins/Bridge/Source/Bridge/Bridge.Build.cs
--- a/Engine/Plugins/Bridge/Source/Bridge/Bridge.Build.cs
+++ b/Engine/Plugins/Bridge/Source/Bridge/Bridge.Build.cs
@@ -39,6 +39,7 @@ namespace UnrealBuildTool.Rules
 					"Networking",
 					"Sockets",
 					"ToolMenus",
+					"MegascansPlugin",
 					});
 
 			PrivateIncludePaths.AddRange(
diff --git a/Engine/Plugins/Bridge/Source/Bridge/Private/UI/BrowserBinding.cpp b/Engine/Plugins/Bridge/Source/Bridge/Private/UI/BrowserBinding.cpp
--- a/Engine/Plugins/Bridge/Source/Bridge/Private/UI/BrowserBinding.cpp
+++ b/Engine/Plugins/Bridge/Source/Bridge/Private/UI/BrowserBinding.cpp
@@ -6,6 +6,7 @@
 #include "WebBrowserModule.h"
 #include "IWebBrowserSingleton.h"
 #include "IWebBrowserCookieManager.h"
+#include "SMSWindow.h"
 
 
 UBrowserBinding::UBrowserBinding(const FObjectInitializer& ObjectInitializer)
@@ -206,3 +207,8 @@ void UBrowserBinding::Logout()
 	
 }
 
+void UBrowserBinding::OpenMegascansPluginSettings()
+{
+	MegascansSettingsWindow::OpenSettingsWindow();
+}
+
diff --git a/Engine/Plugins/Bridge/Source/Bridge/Private/UI/BrowserBinding.h b/Engine/Plugins/Bridge/Source/Bridge/Private/UI/BrowserBinding.h
--- a/Engine/Plugins/Bridge/Source/Bridge/Private/UI/BrowserBinding.h
+++ b/Engine/Plugins/Bridge/Source/Bridge/Private/UI/BrowserBinding.h
@@ -70,5 +70,8 @@ public:
 	UFUNCTION()
 	void Logout();
 
+	UFUNCTION()
+	void OpenMegascansPluginSettings();
+
 	TSharedRef<FBridgeMessageHandler> BridgeMessageHandler = MakeShared<FBridgeMessageHandler>();
 };
