Potential fix for Megascans CIS fails.

diff --git a/Engine/Plugins/Bridge/Bridge.uplugin b/Engine/Plugins/Bridge/Bridge.uplugin
--- a/Engine/Plugins/Bridge/Bridge.uplugin
+++ b/Engine/Plugins/Bridge/Bridge.uplugin
@@ -25,7 +25,17 @@
 			"WhitelistPlatforms": [
 				"Win64",
 				"Mac"
-				
+
+			]
+		},
+		{
+			"Name": "MegascansPlugin",
+			"Type": "Editor",
+			"LoadingPhase": "PostDefault",
+			"WhitelistPlatforms": [
+				"Win64",
+				"Mac",
+				"Linux"
 			]
 		}
 	],
