Fix Compilation tests. #jira UE-109590 UE-109585

diff --git a/Engine/Plugins/Bridge/Source/Bridge/Private/UI/BridgeUIManager.h b/Engine/Plugins/Bridge/Source/Bridge/Private/UI/BridgeUIManager.h
--- a/Engine/Plugins/Bridge/Source/Bridge/Private/UI/BridgeUIManager.h
+++ b/Engine/Plugins/Bridge/Source/Bridge/Private/UI/BridgeUIManager.h
@@ -5,6 +5,9 @@
 #include "Widgets/Docking/SDockTab.h"
 #include "SWebBrowser.h"
 #include "UI/BrowserBinding.h"
+#include "Framework/MultiBox/MultiBoxBuilder.h"
+#include "Framework/Docking/TabManager.h"
+
 
 class FBridgeUIManagerImpl;
 class FArguments;
diff --git a/Engine/Plugins/Bridge/Source/Bridge/Private/UI/FBridgeMessageHandler.cpp b/Engine/Plugins/Bridge/Source/Bridge/Private/UI/FBridgeMessageHandler.cpp
--- a/Engine/Plugins/Bridge/Source/Bridge/Private/UI/FBridgeMessageHandler.cpp
+++ b/Engine/Plugins/Bridge/Source/Bridge/Private/UI/FBridgeMessageHandler.cpp
@@ -3,6 +3,10 @@
 #include "UI/FBridgeMessageHandler.h"
 #include "UI/BrowserBinding.h"
 #include "UI/BridgeUIManager.h"
+#include "Framework/Application/SlateApplication.h"
+#include "Layout/Geometry.h"
+#include "Widgets/SWindow.h"
+
 
 FBridgeMessageHandler::FBridgeMessageHandler(const TSharedPtr<FGenericApplicationMessageHandler>& InTargetHandler)
 	: TargetHandler(InTargetHandler)
@@ -129,7 +133,7 @@ bool FBridgeMessageHandler::OnMouseUp(const EMouseButtons::Type Button, const FV
 
 		if (!WithInBrowserWidth || !WithInBrowserHeight)
 		{
-			UE_LOG(LogTemp, Error, TEXT("Dropped Outside"));
+			
 			FBridgeUIManager::BrowserBinding->OnDroppedDelegate.Execute(TEXT("test"));
 
 			// if (WithInViewPortWidth && WithInViewPortHeight)
@@ -139,7 +143,7 @@ bool FBridgeMessageHandler::OnMouseUp(const EMouseButtons::Type Button, const FV
 		}
 		else
 		{
-			UE_LOG(LogTemp, Error, TEXT("Dropped Outside"));
+			
 			FBridgeUIManager::BrowserBinding->OnDropDiscardedDelegate.Execute(TEXT("test"));
 		}
 
@@ -195,7 +199,7 @@ bool FBridgeMessageHandler::OnMouseWheel(const float Delta, const FVector2D Curs
 
 bool FBridgeMessageHandler::OnMouseMove()
 {
-	// UE_LOG(LogTemp, Error, TEXT("Pointer Event! %f"), FSlateApplication::Get().GetCursorPos().X);
+	
 	SWindow* DragDropWindow = FBridgeUIManager::Instance->DragDropWindow.Get();
 
 	FVector2D DragDropWindowSize = DragDropWindow->GetTickSpaceGeometry().GetAbsoluteSize();
diff --git a/Engine/Plugins/Bridge/Source/MegascansPlugin/Private/AssetImporters/DHIImport.cpp b/Engine/Plugins/Bridge/Source/MegascansPlugin/Private/AssetImporters/DHIImport.cpp
--- a/Engine/Plugins/Bridge/Source/MegascansPlugin/Private/AssetImporters/DHIImport.cpp
+++ b/Engine/Plugins/Bridge/Source/MegascansPlugin/Private/AssetImporters/DHIImport.cpp
@@ -7,6 +7,11 @@
 #include "Misc/ScopedSlowTask.h"
 #include "EditorAssetLibrary.h"
 #include "GenericPlatform/GenericPlatformFile.h"
+#include "HAL/PlatformFilemanager.h"
+#include "Misc/MessageDialog.h"
+#include "Misc/Paths.h"
+#include "Internationalization/Text.h"
+#include "UObject/Object.h"
 
 TSharedPtr<FImportDHI> FImportDHI::ImportDHIInst;
 
diff --git a/Engine/Plugins/Bridge/Source/MegascansPlugin/Private/AssetImporters/ProgressiveImportSurfaces.cpp b/Engine/Plugins/Bridge/Source/MegascansPlugin/Private/AssetImporters/ProgressiveImportSurfaces.cpp
--- a/Engine/Plugins/Bridge/Source/MegascansPlugin/Private/AssetImporters/ProgressiveImportSurfaces.cpp
+++ b/Engine/Plugins/Bridge/Source/MegascansPlugin/Private/AssetImporters/ProgressiveImportSurfaces.cpp
@@ -65,6 +65,12 @@ void FImportProgressiveSurfaces::ImportAsset(TSharedPtr<FJsonObject> AssetImport
 	FString DestinationFolder = FPaths::Combine(FPaths::ProjectContentDir(), DestinationPath.Replace(TEXT("/Game/"), TEXT("")));
 
 	CopyUassetFiles(ImportData->FilePaths, DestinationFolder);
+
+	if (AssetMetaData.assetSubType == TEXT("imperfection") && ImportData->ProgressiveStage == 3)
+	{
+		ImportData->ProgressiveStage = 4;
+	}
+
 	
 
 	if (!PreviewDetails.Contains(ImportData->AssetId))
@@ -86,20 +92,28 @@ void FImportProgressiveSurfaces::ImportAsset(TSharedPtr<FJsonObject> AssetImport
 	}
 	else if (ImportData->ProgressiveStage == 2)
 	{
-		FString AlbedoPath = TEXT("");
-		FString TextureType = TEXT("albedo");
+		FString TexturePath = TEXT("");
+		FString TextureType = TEXT("");
+		if (AssetMetaData.assetSubType == TEXT("imperfection"))
+		{
+			TextureType = TEXT("roughness");
+		}
+		else
+		{
+			TextureType = TEXT("albedo");
+		}
 
 		for (FTexturesList TextureMeta : AssetMetaData.textureSets)
 		{
-			if (TextureMeta.type == TEXT("albedo"))
+			if (TextureMeta.type == TextureType)
 			{
-				AlbedoPath = TextureMeta.path;
+				TexturePath = TextureMeta.path;
 			}
 		}		
 
-		FAssetData AlbedoData = AssetRegistry.GetAssetByObjectPath(FName(*AlbedoPath));
-		FSoftObjectPath ItemToStream = AlbedoData.ToSoftObjectPath();
-		Streamable.RequestAsyncLoad(ItemToStream, FStreamableDelegate::CreateRaw(this, &FImportProgressiveSurfaces::HandlePreviewTextureLoad, AlbedoData, ImportData->AssetId, TextureType));
+		FAssetData TextureData = AssetRegistry.GetAssetByObjectPath(FName(*TexturePath));
+		FSoftObjectPath ItemToStream = TextureData.ToSoftObjectPath();
+		Streamable.RequestAsyncLoad(ItemToStream, FStreamableDelegate::CreateRaw(this, &FImportProgressiveSurfaces::HandlePreviewTextureLoad, TextureData, ImportData->AssetId, TextureType));
 
 
 	}
@@ -112,7 +126,7 @@ void FImportProgressiveSurfaces::ImportAsset(TSharedPtr<FJsonObject> AssetImport
 
 		for (FTexturesList TextureMeta : AssetMetaData.textureSets)
 		{
-			if (TextureMeta.type == TEXT("normal"))
+			if (TextureMeta.type == TextureType)
 			{
 				NormalPath = TextureMeta.path;
 			}
@@ -212,3 +226,4 @@ void FImportProgressiveSurfaces::HandleHighInstanceLoad(FAssetData HighInstanceD
 	PreviewDetails[AssetID]->ActorInLevel->GetStaticMeshComponent()->SetMaterial(0, CastChecked<UMaterialInterface>(HighInstanceData.GetAsset()));
 	PreviewDetails.Remove(AssetID);
 }
+
diff --git a/Engine/Plugins/Bridge/Source/MegascansPlugin/Private/UI/SMSWindow.cpp b/Engine/Plugins/Bridge/Source/MegascansPlugin/Private/UI/SMSWindow.cpp
--- a/Engine/Plugins/Bridge/Source/MegascansPlugin/Private/UI/SMSWindow.cpp
+++ b/Engine/Plugins/Bridge/Source/MegascansPlugin/Private/UI/SMSWindow.cpp
@@ -11,6 +11,8 @@
 #include "UObject/UObjectGlobals.h"
 #include "Math/Vector2D.h"
 #include "EditorAssetLibrary.h"
+#include "Modules/ModuleManager.h"
+
 
 #define LOCTEXT_NAMESPACE "MegascansSettings"
 
diff --git a/Engine/Plugins/Bridge/Source/MegascansPlugin/Private/Utilities/MaterialUtils.cpp b/Engine/Plugins/Bridge/Source/MegascansPlugin/Private/Utilities/MaterialUtils.cpp
--- a/Engine/Plugins/Bridge/Source/MegascansPlugin/Private/Utilities/MaterialUtils.cpp
+++ b/Engine/Plugins/Bridge/Source/MegascansPlugin/Private/Utilities/MaterialUtils.cpp
@@ -18,7 +18,7 @@
 #include "GameFramework/Actor.h"
 #include "Engine/Level.h"
 #include "Components/StaticMeshComponent.h"
-
+#include "Editor/EditorEngine.h"
 
 
 UMaterialInstanceConstant* FMaterialUtils::CreateInstanceMaterial(const FString& MasterMaterialPath, const FString& InstanceDestination, const FString& AssetName)
diff --git a/Engine/Plugins/Bridge/Source/MegascansPlugin/Private/Utilities/MiscUtils.cpp b/Engine/Plugins/Bridge/Source/MegascansPlugin/Private/Utilities/MiscUtils.cpp
--- a/Engine/Plugins/Bridge/Source/MegascansPlugin/Private/Utilities/MiscUtils.cpp
+++ b/Engine/Plugins/Bridge/Source/MegascansPlugin/Private/Utilities/MiscUtils.cpp
@@ -35,6 +35,8 @@
 #include "InstancedFoliageActor.h"
 #include "FoliageType_InstancedStaticMesh.h"
 #include "Engine/World.h"
+#include "Editor/EditorEngine.h"
+
 
 TSharedPtr<FJsonObject> DeserializeJson(const FString& JsonStringData)
 {
