Cross-compile linux potential fix.

diff --git a/Engine/Plugins/Bridge/Bridge.uplugin b/Engine/Plugins/Bridge/Bridge.uplugin
--- a/Engine/Plugins/Bridge/Bridge.uplugin
+++ b/Engine/Plugins/Bridge/Bridge.uplugin
@@ -14,8 +14,8 @@
 	"EnabledByDefault": true,
 	"SupportedTargetPlatforms": [
 		"Win64",
-		"Mac",
-		"Linux"
+		"Mac"
+		
 	],
 	"Modules": [
 		{
diff --git a/Engine/Plugins/Bridge/Source/MegascansPlugin/Private/AssetImporters/DHIImport.cpp b/Engine/Plugins/Bridge/Source/MegascansPlugin/Private/AssetImporters/DHIImport.cpp
--- a/Engine/Plugins/Bridge/Source/MegascansPlugin/Private/AssetImporters/DHIImport.cpp
+++ b/Engine/Plugins/Bridge/Source/MegascansPlugin/Private/AssetImporters/DHIImport.cpp
@@ -7,7 +7,7 @@
 #include "Misc/ScopedSlowTask.h"
 #include "EditorAssetLibrary.h"
 #include "GenericPlatform/GenericPlatformFile.h"
-#include "HAL/PlatformFilemanager.h"
+//#include "HAL/PlatformFilemanager.h"
 #include "Misc/MessageDialog.h"
 #include "Misc/Paths.h"
 #include "Internationalization/Text.h"
@@ -62,7 +62,7 @@ void FImportDHI::ImportAsset(TSharedPtr<FJsonObject> AssetImportJson)
 
 	TArray<FString> SourceCommonFiles;
 	PlatformFile.FindFilesRecursively(SourceCommonFiles, *CharacterSourceData->CommonPath, NULL);
-	TArray<FString> ExistingAssets = UEditorAssetLibrary::ListAssets(TEXT("/Game/MetaHumans/32222222222222222222222222222222222222222222222222"));
+	TArray<FString> ExistingAssets = UEditorAssetLibrary::ListAssets(TEXT("/Game/MetaHumans/Common"));
 
 	FString ProjectCommonPath = FPaths::ConvertRelativePathToFull(FPaths::Combine(FPaths::ProjectContentDir(), TEXT("MetaHumans"), TEXT("Common")));
 	FString SourceCommonPath = CharacterSourceData->CommonPath;
