Fix missing headers for Compilation tests. #jira UE-109716 UE-109590

diff --git a/Engine/Plugins/Bridge/Source/MegascansPlugin/Private/AssetImporters/DHIImport.cpp b/Engine/Plugins/Bridge/Source/MegascansPlugin/Private/AssetImporters/DHIImport.cpp
--- a/Engine/Plugins/Bridge/Source/MegascansPlugin/Private/AssetImporters/DHIImport.cpp
+++ b/Engine/Plugins/Bridge/Source/MegascansPlugin/Private/AssetImporters/DHIImport.cpp
@@ -7,7 +7,7 @@
 #include "Misc/ScopedSlowTask.h"
 #include "EditorAssetLibrary.h"
 #include "GenericPlatform/GenericPlatformFile.h"
-//#include "HAL/PlatformFilemanager.h"
+#include "HAL/PlatformFilemanager.h"
 #include "Misc/MessageDialog.h"
 #include "Misc/Paths.h"
 #include "Internationalization/Text.h"
diff --git a/Engine/Plugins/Bridge/Source/MegascansPlugin/Private/AssetImporters/ProgressiveImport3D.cpp b/Engine/Plugins/Bridge/Source/MegascansPlugin/Private/AssetImporters/ProgressiveImport3D.cpp
--- a/Engine/Plugins/Bridge/Source/MegascansPlugin/Private/AssetImporters/ProgressiveImport3D.cpp
+++ b/Engine/Plugins/Bridge/Source/MegascansPlugin/Private/AssetImporters/ProgressiveImport3D.cpp
@@ -132,6 +132,7 @@ void FImportProgressive3D::ImportAsset(TSharedPtr<FJsonObject> AssetImportJson)
 		FAssetData MInstanceData = AssetRegistry.GetAssetByObjectPath(FName(*MInstancePath));
 
 		FSoftObjectPath ItemToStream = PreviewMeshData.ToSoftObjectPath();
+		//UE_LOG(LogTemp, Error, TEXT("Instance path : %s"), *MInstancePath);
 		
 		Streamable.RequestAsyncLoad(ItemToStream, FStreamableDelegate::CreateRaw(this, &FImportProgressive3D::HandlePreviewInstanceLoad, MInstanceData, ImportData->AssetId));
 		SpawnAtCenter(PreviewMeshData, ImportData);
@@ -150,6 +151,7 @@ void FImportProgressive3D::ImportAsset(TSharedPtr<FJsonObject> AssetImportJson)
 			}
 		}
 
+		//UE_LOG(LogTemp, Error, TEXT("Albedo path : %s"), *AlbedoPath);
 		
 
 		FAssetData AlbedoData = AssetRegistry.GetAssetByObjectPath(FName(*AlbedoPath));
@@ -208,6 +210,7 @@ void FImportProgressive3D::HandlePreviewTextureLoad(FAssetData TextureData, FStr
 
 void FImportProgressive3D::HandlePreviewInstanceLoad(FAssetData PreviewInstanceData, FString AssetID)
 {
+	
 	PreviewDetails[AssetID]->PreviewInstance = Cast<UMaterialInstanceConstant>(PreviewInstanceData.GetAsset());
 }
 
@@ -221,6 +224,10 @@ void FImportProgressive3D::HandleHighAssetLoad(FAssetData HighAssetData, FString
 	AStaticMeshActor* PreviewActor = ProgressiveData[AssetID];
 	UStaticMesh* SourceMesh = Cast<UStaticMesh>(HighAssetData.GetAsset());
 
+	/*SourceMesh->OnPostMeshBuild().AddLambda([this](UStaticMesh* StaticMesh) {
+		UE_LOG(LogTemp, Error, TEXT("Data build complete..."));
+	});*/
+
 	if (FMaterialUtils::ShouldOverrideMaterial(AssetMetaData.assetType))
 	{
 		AssetUtils::DeleteAsset(AssetMetaData.materialInstances[0].instancePath);
@@ -240,6 +247,7 @@ void FImportProgressive3D::AsyncCacheData(FAssetData HighAssetData, FString Asse
 {
 
 	UStaticMesh* SourceMesh = Cast<UStaticMesh>(HighAssetData.GetAsset());
+	
 
 	if (bWaitNaniteConversion)
 	{
diff --git a/Engine/Plugins/Bridge/Source/MegascansPlugin/Private/Utilities/MaterialUtils.cpp b/Engine/Plugins/Bridge/Source/MegascansPlugin/Private/Utilities/MaterialUtils.cpp
--- a/Engine/Plugins/Bridge/Source/MegascansPlugin/Private/Utilities/MaterialUtils.cpp
+++ b/Engine/Plugins/Bridge/Source/MegascansPlugin/Private/Utilities/MaterialUtils.cpp
@@ -20,6 +20,9 @@
 #include "Components/StaticMeshComponent.h"
 #include "Editor/EditorEngine.h"
 
+#include "EngineGlobals.h"
+#include "Editor.h"
+
 
 UMaterialInstanceConstant* FMaterialUtils::CreateInstanceMaterial(const FString& MasterMaterialPath, const FString& InstanceDestination, const FString& AssetName)
 {	
diff --git a/Engine/Plugins/Bridge/Source/MegascansPlugin/Private/Utilities/MiscUtils.cpp b/Engine/Plugins/Bridge/Source/MegascansPlugin/Private/Utilities/MiscUtils.cpp
--- a/Engine/Plugins/Bridge/Source/MegascansPlugin/Private/Utilities/MiscUtils.cpp
+++ b/Engine/Plugins/Bridge/Source/MegascansPlugin/Private/Utilities/MiscUtils.cpp
@@ -37,6 +37,9 @@
 #include "Engine/World.h"
 #include "Editor/EditorEngine.h"
 
+#include "EngineGlobals.h"
+#include "Editor.h"
+
 
 TSharedPtr<FJsonObject> DeserializeJson(const FString& JsonStringData)
 {
