Fixed casing of HAL/PlatformFileManager.h include to fix Linux cross-compile.

#trivial
#rnx

diff --git a/Engine/Plugins/Bridge/Source/MegascansPlugin/Private/AssetImporters/DHIImport.cpp b/Engine/Plugins/Bridge/Source/MegascansPlugin/Private/AssetImporters/DHIImport.cpp
--- a/Engine/Plugins/Bridge/Source/MegascansPlugin/Private/AssetImporters/DHIImport.cpp
+++ b/Engine/Plugins/Bridge/Source/MegascansPlugin/Private/AssetImporters/DHIImport.cpp
@@ -7,7 +7,7 @@
 #include "Misc/ScopedSlowTask.h"
 #include "EditorAssetLibrary.h"
 #include "GenericPlatform/GenericPlatformFile.h"
-#include "HAL/PlatformFilemanager.h"
+#include "HAL/PlatformFileManager.h"
 #include "Misc/MessageDialog.h"
 #include "Misc/Paths.h"
 #include "Internationalization/Text.h"
