Quixel Bridge export and progressive exports.

diff --git a/Engine/Plugins/Bridge/Source/Bridge/Private/UI/BridgeUIManager.h b/Engine/Plugins/Bridge/Source/Bridge/Private/UI/BridgeUIManager.h
--- a/Engine/Plugins/Bridge/Source/Bridge/Private/UI/BridgeUIManager.h
+++ b/Engine/Plugins/Bridge/Source/Bridge/Private/UI/BridgeUIManager.h
@@ -20,7 +20,7 @@ public:
 	void Shutdown();
 	void HandleBrowserUrlChanged(const FText& Url);
 	TSharedPtr<SWebBrowser> WebBrowserWidget;
-	TSharedPtr<SWindow> DragDropWindow;
+	TArray<TSharedPtr<SWindow>> DragDropWindows;
 
 private:
 	void SetupMenuItem();
diff --git a/Engine/Plugins/Bridge/Source/Bridge/Private/UI/BrowserBinding.cpp b/Engine/Plugins/Bridge/Source/Bridge/Private/UI/BrowserBinding.cpp
--- a/Engine/Plugins/Bridge/Source/Bridge/Private/UI/BrowserBinding.cpp
+++ b/Engine/Plugins/Bridge/Source/Bridge/Private/UI/BrowserBinding.cpp
@@ -155,27 +155,34 @@ void UBrowserBinding::OpenExternalUrl(FString Url)
 	FPlatformProcess::LaunchURL(*Url, nullptr, nullptr);
 }
 
-void UBrowserBinding::DragStarted(FString ImageUrl)
+void UBrowserBinding::DragStarted(TArray<FString> ImageUrls)
 {
-	TSharedPtr<SWebBrowser> PopupWebBrowser = SNew(SWebBrowser)
-								.ShowControls(false);
-				
-	FBridgeUIManager::Instance->DragDropWindow = SNew(SWindow)
-		.ClientSize(FVector2D(120, 120))
-		.SupportsTransparency(EWindowTransparency::PerWindow)
-		.RenderOpacity(0.3)
-		.CreateTitleBar(false)
-		.IsTopmostWindow(true)
-		.FocusWhenFirstShown(false)
-		.SupportsMaximize(false)
-		.SupportsMinimize(false)
-		[
-			PopupWebBrowser.ToSharedRef()
-		];
+	for (int32 index = ImageUrls.Num() - 1; index >= 0; index--)
+	{
+		FString ImageUrl = ImageUrls[index];
+
+	 	// UE_LOG(LogTemp, Error, TEXT("Image: %s"), *ImageUrl);
+		TSharedPtr<SWebBrowser> PopupWebBrowser = SNew(SWebBrowser)
+									.ShowControls(false);
+		TSharedPtr<SWindow> DragDropWindow = SNew(SWindow)
+			.ClientSize(FVector2D(120, 120))
+			.SupportsTransparency(EWindowTransparency::PerWindow)
+			.RenderOpacity(0.3)
+			.CreateTitleBar(false)
+			.IsTopmostWindow(true)
+			.FocusWhenFirstShown(false)
+			.SupportsMaximize(false)
+			.SupportsMinimize(false)
+			[
+				PopupWebBrowser.ToSharedRef()
+			];
+					
+		FBridgeUIManager::Instance->DragDropWindows.Add(DragDropWindow);
 
-	PopupWebBrowser->LoadString(FString::Printf(TEXT("<!DOCTYPE html><html lang=\"en\"> <head> <meta charset=\"UTF-8\" /> <meta name=\"viewport\" content=\"width=device-width, initial-scale=1.0\" /> <style> * { padding: 0px; } body { padding: 0px; margin: 0px; } #container { display: flex; width: 120px; height: 120px; background: #202020; justify-content: center; align-items: center;} #full-image { max-width: 120px; height: auto; display: block; font-size: 0; } </style> </head> <body> <div id=\"container\"> <img id=\"full-image\" src=\"%s\" /> </div> </body></html>"), *ImageUrl), TEXT(""));
+		PopupWebBrowser->LoadString(FString::Printf(TEXT("<!DOCTYPE html><html lang=\"en\"> <head> <meta charset=\"UTF-8\" /> <meta name=\"viewport\" content=\"width=device-width, initial-scale=1.0\" /> <style> * { padding: 0px; } body { padding: 0px; margin: 0px; } #container { display: flex; width: 120px; height: 120px; background: #202020; justify-content: center; align-items: center;} #full-image { max-width: 110px; max-height: 110px; display: block; font-size: 0; } </style> </head> <body> <div id=\"container\"> <img id=\"full-image\" src=\"%s\" /> </div> </body></html>"), *ImageUrl), TEXT(""));
 
-	FSlateApplication::Get().AddWindow(FBridgeUIManager::Instance->DragDropWindow.ToSharedRef());
+		FSlateApplication::Get().AddWindow(DragDropWindow.ToSharedRef(), false);
+	}
 
 	TSharedRef<FGenericApplicationMessageHandler> TargetHandler = FSlateApplication::Get().GetPlatformApplication().Get()->GetMessageHandler();
 	UBrowserBinding::BridgeMessageHandler->SetTargetHandler(TargetHandler);
@@ -187,8 +194,6 @@ void UBrowserBinding::DragStarted(FString ImageUrl)
 	FSlateApplication::Get().SpawnToolTip(MyToolTip.ToSharedRef(), FVector2D(200, 300));*/
 
 	//	FBridgeUIManager::Instance->DragDropWindow.Get()->Resize(FVector2D(100, 100));
-
-
 }
 
 void UBrowserBinding::Logout()
@@ -203,8 +208,6 @@ void UBrowserBinding::Logout()
 			CookieManager->DeleteCookies();
 		}
 	}
-
-	
 }
 
 void UBrowserBinding::OpenMegascansPluginSettings()
diff --git a/Engine/Plugins/Bridge/Source/Bridge/Private/UI/BrowserBinding.h b/Engine/Plugins/Bridge/Source/Bridge/Private/UI/BrowserBinding.h
--- a/Engine/Plugins/Bridge/Source/Bridge/Private/UI/BrowserBinding.h
+++ b/Engine/Plugins/Bridge/Source/Bridge/Private/UI/BrowserBinding.h
@@ -57,7 +57,7 @@ public:
 	void ShowDialog(FString Type, FString Url);
 
 	UFUNCTION()
-	void DragStarted(FString ImageUrl);
+	void DragStarted(TArray<FString> ImageUrl);
 
 	void ShowLoginDialog(FString Url);
 
diff --git a/Engine/Plugins/Bridge/Source/Bridge/Private/UI/FBridgeMessageHandler.cpp b/Engine/Plugins/Bridge/Source/Bridge/Private/UI/FBridgeMessageHandler.cpp
--- a/Engine/Plugins/Bridge/Source/Bridge/Private/UI/FBridgeMessageHandler.cpp
+++ b/Engine/Plugins/Bridge/Source/Bridge/Private/UI/FBridgeMessageHandler.cpp
@@ -108,8 +108,16 @@ bool FBridgeMessageHandler::OnMouseUp(const EMouseButtons::Type Button, const FV
 {
 	if (TargetHandler.IsValid())
 	{
-		// Destroy the drag popup
-		FBridgeUIManager::Instance->DragDropWindow.Get()->RequestDestroyWindow();
+		// Destroy the drag popups
+		for (auto& Window : FBridgeUIManager::Instance->DragDropWindows)
+		{
+			Window->RequestDestroyWindow();
+		}
+
+		FBridgeUIManager::Instance->DragDropWindows.RemoveAll([](TSharedPtr<SWindow> Val)
+		{
+			 return true;
+		});
 
 		// Get browser dimensions
 		FGeometry BrowserGeometry = FBridgeUIManager::Instance->WebBrowserWidget.Get()->GetTickSpaceGeometry();
@@ -199,12 +207,18 @@ bool FBridgeMessageHandler::OnMouseWheel(const float Delta, const FVector2D Curs
 
 bool FBridgeMessageHandler::OnMouseMove()
 {
-	
-	SWindow* DragDropWindow = FBridgeUIManager::Instance->DragDropWindow.Get();
+	for (int32 Index = 0; Index !=  FBridgeUIManager::Instance->DragDropWindows.Num(); ++Index)
+	{
+		TSharedPtr<SWindow> DragDropWindow = FBridgeUIManager::Instance->DragDropWindows[Index];
+		FVector2D DragDropWindowSize = DragDropWindow->GetTickSpaceGeometry().GetAbsoluteSize();
+		FVector2D CursorPosition = FSlateApplication::Get().GetCursorPos();
+		DragDropWindow.Get()->MoveWindowTo(FVector2D(CursorPosition.X + (10 * Index) - (DragDropWindowSize.X / 2), CursorPosition.Y + (10 * Index) - (DragDropWindowSize.Y / 2)));
+		if (!DragDropWindow->IsVisible())
+		{
+			DragDropWindow->ShowWindow();
+		}
+	}
 
-	FVector2D DragDropWindowSize = DragDropWindow->GetTickSpaceGeometry().GetAbsoluteSize();
-	FVector2D CursorPosition = FSlateApplication::Get().GetCursorPos();
-	FBridgeUIManager::Instance->DragDropWindow.Get()->MoveWindowTo(FVector2D(CursorPosition.X - (DragDropWindowSize.X / 2), CursorPosition.Y - (DragDropWindowSize.Y / 2)));
 	return TargetHandler->OnMouseMove();
 }
 
diff --git a/Engine/Plugins/Bridge/Source/MegascansPlugin/Private/AssetImporters/ProgressiveImport3D.cpp b/Engine/Plugins/Bridge/Source/MegascansPlugin/Private/AssetImporters/ProgressiveImport3D.cpp
--- a/Engine/Plugins/Bridge/Source/MegascansPlugin/Private/AssetImporters/ProgressiveImport3D.cpp
+++ b/Engine/Plugins/Bridge/Source/MegascansPlugin/Private/AssetImporters/ProgressiveImport3D.cpp
@@ -29,6 +29,8 @@
 #include "Components/StaticMeshComponent.h"
 #include "Misc/Paths.h"
 
+#include "StaticMeshCompiler.h"
+
 
 
 TSharedPtr<FImportProgressive3D> FImportProgressive3D::ImportProgressive3DInst;
@@ -44,7 +46,7 @@ TSharedPtr<FImportProgressive3D> FImportProgressive3D::Get()
 }
 
 
-void FImportProgressive3D::SpawnAtCenter(FAssetData AssetData, TSharedPtr<FUAssetData> ImportData)
+void FImportProgressive3D::SpawnAtCenter(FAssetData AssetData, TSharedPtr<FUAssetData> ImportData, float LocationOffset)
 {
 	FViewport* ActiveViewport = GEditor->GetActiveViewport();
 	FEditorViewportClient* EditorViewClient = (FEditorViewportClient*)ActiveViewport->GetClient();
@@ -54,7 +56,7 @@ void FImportProgressive3D::SpawnAtCenter(FAssetData AssetData, TSharedPtr<FUAsse
 	FRotator InitialRotation(0.0f, 0.0f, 0.0f);
 
 	FVector SpawnLocation = ViewPosition + (ViewDirection * 300.0f) ;
-	SpawnLocation.Y += LocationOffset;
+	SpawnLocation.X += LocationOffset;
 
 
 	FVector Location(0.0f, 0.0f, 0.0f);
@@ -79,12 +81,12 @@ void FImportProgressive3D::SpawnAtCenter(FAssetData AssetData, TSharedPtr<FUAsse
 		ProgressiveData.Add(ImportData->AssetId, SMActor);
 	}
 
-	SetLocationOffset(0.0f);
+	//SetLocationOffset(0.0f);
 
 
 }
 
-void FImportProgressive3D::ImportAsset(TSharedPtr<FJsonObject> AssetImportJson)
+void FImportProgressive3D::ImportAsset(TSharedPtr<FJsonObject> AssetImportJson, float LocationOffset)
 {
 
 	FAssetRegistryModule& AssetRegistryModule = FModuleManager::GetModuleChecked<FAssetRegistryModule>("AssetRegistry");
@@ -115,27 +117,20 @@ void FImportProgressive3D::ImportAsset(TSharedPtr<FJsonObject> AssetImportJson)
 
 
 	if (ImportData->ProgressiveStage == 1)
-	{
-		
+	{		
 
 		FString MeshPath = AssetMetaData.meshList[0].path;
-		
-
-		PreviewDetails[ImportData->AssetId]->PreviewMeshPath = MeshPath;	
-
-
+		PreviewDetails[ImportData->AssetId]->PreviewMeshPath = MeshPath;
 		FAssetData PreviewMeshData = AssetRegistry.GetAssetByObjectPath(FName(*MeshPath));
-		
-
 		FString MInstancePath = AssetMetaData.materialInstances[0].instancePath ;
 		
 		FAssetData MInstanceData = AssetRegistry.GetAssetByObjectPath(FName(*MInstancePath));
-
 		FSoftObjectPath ItemToStream = PreviewMeshData.ToSoftObjectPath();
-		//UE_LOG(LogTemp, Error, TEXT("Instance path : %s"), *MInstancePath);
+		
+		if (!MInstanceData.IsValid()) return;
 		
 		Streamable.RequestAsyncLoad(ItemToStream, FStreamableDelegate::CreateRaw(this, &FImportProgressive3D::HandlePreviewInstanceLoad, MInstanceData, ImportData->AssetId));
-		SpawnAtCenter(PreviewMeshData, ImportData);
+		SpawnAtCenter(PreviewMeshData, ImportData, LocationOffset);
 
 	}
 	else if (ImportData->ProgressiveStage == 2)
@@ -151,11 +146,12 @@ void FImportProgressive3D::ImportAsset(TSharedPtr<FJsonObject> AssetImportJson)
 			}
 		}
 
-		//UE_LOG(LogTemp, Error, TEXT("Albedo path : %s"), *AlbedoPath);
 		
-
 		FAssetData AlbedoData = AssetRegistry.GetAssetByObjectPath(FName(*AlbedoPath));
 		FSoftObjectPath ItemToStream = AlbedoData.ToSoftObjectPath();
+
+		if (!AlbedoData.IsValid()) return;
+
 		Streamable.RequestAsyncLoad(ItemToStream, FStreamableDelegate::CreateRaw(this, &FImportProgressive3D::HandlePreviewTextureLoad, AlbedoData, ImportData->AssetId, TextureType));
 
 
@@ -178,6 +174,9 @@ void FImportProgressive3D::ImportAsset(TSharedPtr<FJsonObject> AssetImportJson)
 		
 		FAssetData NormalData = AssetRegistry.GetAssetByObjectPath(FName(*NormalPath));
 		FSoftObjectPath ItemToStream = NormalData.ToSoftObjectPath();
+
+		if (!NormalData.IsValid()) return;
+
 		Streamable.RequestAsyncLoad(ItemToStream, FStreamableDelegate::CreateRaw(this, &FImportProgressive3D::HandlePreviewTextureLoad, NormalData, ImportData->AssetId, TextureType));
 
 	}
@@ -194,6 +193,9 @@ void FImportProgressive3D::ImportAsset(TSharedPtr<FJsonObject> AssetImportJson)
 		}
 		
 		FSoftObjectPath ItemToStream = HighMeshData.ToSoftObjectPath();
+
+		if (!HighMeshData.IsValid()) return;
+
 		Streamable.RequestAsyncLoad(ItemToStream, FStreamableDelegate::CreateRaw(this, &FImportProgressive3D::HandleHighAssetLoad, HighMeshData, ImportData->AssetId, AssetMetaData, bWaitNaniteConversion));
 
 	}
@@ -203,7 +205,13 @@ void FImportProgressive3D::ImportAsset(TSharedPtr<FJsonObject> AssetImportJson)
 
 void FImportProgressive3D::HandlePreviewTextureLoad(FAssetData TextureData, FString AssetID,  FString Type)
 {
-	UTexture* PreviewTexture = Cast<UTexture>(TextureData.GetAsset());
+	if (PreviewDetails[AssetID]->PreviewInstance == nullptr || !IsValid(PreviewDetails[AssetID]->PreviewInstance))
+	{
+		return;
+	}
+
+	UTexture* PreviewTexture = Cast<UTexture>(TextureData.GetAsset());	
+
 	UMaterialEditingLibrary::SetMaterialInstanceTextureParameterValue(PreviewDetails[AssetID]->PreviewInstance, FName(*Type), PreviewTexture);
 	AssetUtils::SavePackage(PreviewDetails[AssetID]->PreviewInstance);
 }
@@ -247,15 +255,15 @@ void FImportProgressive3D::AsyncCacheData(FAssetData HighAssetData, FString Asse
 {
 
 	UStaticMesh* SourceMesh = Cast<UStaticMesh>(HighAssetData.GetAsset());
-	
 
-	if (bWaitNaniteConversion)
+	if (FStaticMeshCompilingManager::Get().IsAsyncCompilationAllowed(SourceMesh) && FStaticMeshCompilingManager::Get().IsAsyncStaticMeshCompilationEnabled())
 	{
-		FPlatformProcess::Sleep(45.0f);
+		while (SourceMesh->IsCompiling())
+		{
+			FPlatformProcess::Sleep(1.0f);
+		}
 	}
 
-	
-
 
 	AsyncTask(ENamedThreads::GameThread, [this, HighAssetData, AssetID]() {
 		SwitchHigh(HighAssetData, AssetID);
@@ -269,6 +277,9 @@ void FImportProgressive3D::SwitchHigh(FAssetData HighAssetData, FString AssetID)
 	UGameplayStatics::GetAllActorsOfClass(GEngine->GetWorldContexts()[0].World(), AStaticMeshActor::StaticClass(), FoundActors);
 	
 	UStaticMesh* SourceMesh = Cast<UStaticMesh>(HighAssetData.GetAsset());
+
+	if (!IsValid(SourceMesh)) return;
+
 	/*AStaticMeshActor* PreviewActor = ProgressiveData[AssetID];
 	if (PreviewActor != nullptr) {
 		PreviewActor->GetStaticMeshComponent()->SetStaticMesh(SourceMesh);
@@ -286,8 +297,8 @@ void FImportProgressive3D::SwitchHigh(FAssetData HighAssetData, FString AssetID)
 		}
 	}
 
-	FString AssetPathHigh = FPaths::GetPath(SourceMesh->GetPathName());
-	FString PreviwPath = FPaths::Combine(AssetPathHigh, TEXT("Preview"));
+	/*FString AssetPathHigh = FPaths::GetPath(SourceMesh->GetPathName());
+	FString PreviwPath = FPaths::Combine(AssetPathHigh, TEXT("Preview"));*/
 	//AssetUtils::DeleteDirectory(PreviwPath);
 
 	
@@ -295,10 +306,6 @@ void FImportProgressive3D::SwitchHigh(FAssetData HighAssetData, FString AssetID)
 	PreviewDetails.Remove(AssetID);
 }
 
-void FImportProgressive3D::SetLocationOffset(float ActorLocationOffset)
-{
-	LocationOffset = ActorLocationOffset;
-}
 
 
 
diff --git a/Engine/Plugins/Bridge/Source/MegascansPlugin/Private/AssetImporters/ProgressiveImport3D.h b/Engine/Plugins/Bridge/Source/MegascansPlugin/Private/AssetImporters/ProgressiveImport3D.h
--- a/Engine/Plugins/Bridge/Source/MegascansPlugin/Private/AssetImporters/ProgressiveImport3D.h
+++ b/Engine/Plugins/Bridge/Source/MegascansPlugin/Private/AssetImporters/ProgressiveImport3D.h
@@ -24,7 +24,7 @@ class FImportProgressive3D
 private:
 	FImportProgressive3D() = default;
 	static TSharedPtr<FImportProgressive3D> ImportProgressive3DInst;
-	void SpawnAtCenter(FAssetData AssetData, TSharedPtr<FUAssetData> ImportData);
+	void SpawnAtCenter(FAssetData AssetData, TSharedPtr<FUAssetData> ImportData, float LocationOffset);
 	TMap<FString, AStaticMeshActor*> ProgressiveData;
 	
 	TMap<FString, TSharedPtr<FProgressiveData>> PreviewDetails;
@@ -32,17 +32,17 @@ private:
 
 	void AsyncCacheData(FAssetData HighAssetData, FString AssetID, FUAssetMeta AssetMetaData, bool bWaitNaniteConversion=false);
 	void SwitchHigh(FAssetData HighAssetData, FString AssetID);
-	float LocationOffset = 0.0f;
+	
 
 public:
 	static TSharedPtr<FImportProgressive3D> Get();
-	void ImportAsset(TSharedPtr<FJsonObject> AssetImportJson);
+	void ImportAsset(TSharedPtr<FJsonObject> AssetImportJson, float LocationOffset);
 
 	void HandlePreviewTextureLoad(FAssetData TextureData, FString AssetID, FString Type);
 	void HandlePreviewInstanceLoad(FAssetData PreviewInstanceData, FString AssetID);
 	void HandleHighAssetLoad(FAssetData HighAssetData, FString AssetID, FUAssetMeta AssetMetaData, bool bWaitNaniteConversion=false);
 	
-	void SetLocationOffset(float LocationOffset);
+	
 
 	~FImportProgressive3D() {UE_LOG(LogTemp, Error, TEXT("Destructor calledd..."))};
 
diff --git a/Engine/Plugins/Bridge/Source/MegascansPlugin/Private/AssetImporters/ProgressiveImportSurfaces.cpp b/Engine/Plugins/Bridge/Source/MegascansPlugin/Private/AssetImporters/ProgressiveImportSurfaces.cpp
--- a/Engine/Plugins/Bridge/Source/MegascansPlugin/Private/AssetImporters/ProgressiveImportSurfaces.cpp
+++ b/Engine/Plugins/Bridge/Source/MegascansPlugin/Private/AssetImporters/ProgressiveImportSurfaces.cpp
@@ -45,7 +45,7 @@ TSharedPtr<FImportProgressiveSurfaces> FImportProgressiveSurfaces::Get()
 }
 
 
-void FImportProgressiveSurfaces::ImportAsset(TSharedPtr<FJsonObject> AssetImportJson)
+void FImportProgressiveSurfaces::ImportAsset(TSharedPtr<FJsonObject> AssetImportJson, float LocationOffset)
 {
 
 	FAssetRegistryModule& AssetRegistryModule = FModuleManager::GetModuleChecked<FAssetRegistryModule>("AssetRegistry");
@@ -86,7 +86,7 @@ void FImportProgressiveSurfaces::ImportAsset(TSharedPtr<FJsonObject> AssetImport
 		FAssetData MInstanceData = AssetRegistry.GetAssetByObjectPath(FName(*MInstancePath));
 
 		FSoftObjectPath ItemToStream = MInstanceData.ToSoftObjectPath();
-		Streamable.RequestAsyncLoad(ItemToStream, FStreamableDelegate::CreateRaw(this, &FImportProgressiveSurfaces::HandlePreviewInstanceLoad, MInstanceData, ImportData->AssetId));
+		Streamable.RequestAsyncLoad(ItemToStream, FStreamableDelegate::CreateRaw(this, &FImportProgressiveSurfaces::HandlePreviewInstanceLoad, MInstanceData, ImportData->AssetId, LocationOffset));
 		
 
 	}
@@ -153,6 +153,11 @@ void FImportProgressiveSurfaces::ImportAsset(TSharedPtr<FJsonObject> AssetImport
 
 void FImportProgressiveSurfaces::HandlePreviewTextureLoad(FAssetData TextureData, FString AssetID, FString Type)
 {
+	if (PreviewDetails[AssetID]->PreviewInstance == nullptr || !IsValid(PreviewDetails[AssetID]->PreviewInstance))
+	{
+		return;
+	}
+
 	UTexture* PreviewTexture = Cast<UTexture>(TextureData.GetAsset());
 	UMaterialEditingLibrary::SetMaterialInstanceTextureParameterValue(PreviewDetails[AssetID]->PreviewInstance, FName(*Type), PreviewTexture);
 	AssetUtils::SavePackage(PreviewDetails[AssetID]->PreviewInstance);
@@ -160,16 +165,16 @@ void FImportProgressiveSurfaces::HandlePreviewTextureLoad(FAssetData TextureData
 
 
 
-void FImportProgressiveSurfaces::HandlePreviewInstanceLoad(FAssetData PreviewInstanceData, FString AssetID)
+void FImportProgressiveSurfaces::HandlePreviewInstanceLoad(FAssetData PreviewInstanceData, FString AssetID, float LocationOffset)
 {
 
 	PreviewDetails[AssetID]->PreviewInstance = Cast<UMaterialInstanceConstant>(PreviewInstanceData.GetAsset());
-	SpawnMaterialPreviewActor(AssetID);
+	SpawnMaterialPreviewActor(AssetID, LocationOffset);
 
 
 }
 
-void FImportProgressiveSurfaces::SpawnMaterialPreviewActor(FString AssetID)
+void FImportProgressiveSurfaces::SpawnMaterialPreviewActor(FString AssetID, float LocationOffset)
 {
 	IAssetRegistry& AssetRegistry = FModuleManager::GetModuleChecked<FAssetRegistryModule>("AssetRegistry").Get();
 	FString SphereMeshPath = TEXT("/Engine/BasicShapes/Sphere.Sphere");
@@ -184,6 +189,7 @@ void FImportProgressiveSurfaces::SpawnMaterialPreviewActor(FString AssetID)
 	FRotator InitialRotation(0.0f, 0.0f, 0.0f);
 
 	FVector SpawnLocation = ViewPosition + (ViewDirection * 300.0f);
+	SpawnLocation.X += LocationOffset;
 
 	FVector Location(0.0f, 0.0f, 0.0f);
 	FRotator Rotation(0.0f, 0.0f, 0.0f);
@@ -211,6 +217,7 @@ void FImportProgressiveSurfaces::HandleHighInstanceLoad(FAssetData HighInstanceD
 {
 	if (!PreviewDetails.Contains(AssetID)) return;
 	if (PreviewDetails[AssetID]->ActorInLevel == nullptr) return;
+	if (!IsValid(PreviewDetails[AssetID]->ActorInLevel)) return;
 
 	if (FMaterialUtils::ShouldOverrideMaterial(AssetMetaData.assetType))
 	{
diff --git a/Engine/Plugins/Bridge/Source/MegascansPlugin/Private/AssetImporters/ProgressiveImportSurfaces.h b/Engine/Plugins/Bridge/Source/MegascansPlugin/Private/AssetImporters/ProgressiveImportSurfaces.h
--- a/Engine/Plugins/Bridge/Source/MegascansPlugin/Private/AssetImporters/ProgressiveImportSurfaces.h
+++ b/Engine/Plugins/Bridge/Source/MegascansPlugin/Private/AssetImporters/ProgressiveImportSurfaces.h
@@ -24,14 +24,14 @@ private:
 	static TSharedPtr<FImportProgressiveSurfaces> ImportProgressiveSurfacesInst;
 
 	TMap<FString, TSharedPtr<FProgressiveSurfaces>> PreviewDetails;
-	void SpawnMaterialPreviewActor(FString AssetID);	
+	void SpawnMaterialPreviewActor(FString AssetID, float LocationOffset);
 
 
 public:
 	static TSharedPtr<FImportProgressiveSurfaces> Get();
-	void ImportAsset(TSharedPtr<FJsonObject> AssetImportJson);
+	void ImportAsset(TSharedPtr<FJsonObject> AssetImportJson, float LocationOffset);
 
-	void HandlePreviewInstanceLoad(FAssetData PreviewInstanceData, FString AssetID);
+	void HandlePreviewInstanceLoad(FAssetData PreviewInstanceData, FString AssetID, float LocationOffset);
 	void HandlePreviewTextureLoad(FAssetData TextureData, FString AssetID, FString Type);
 
 	void HandleHighInstanceLoad(FAssetData HighInstanceData, FString AssetID, FUAssetMeta AssetMetaData);
diff --git a/Engine/Plugins/Bridge/Source/MegascansPlugin/Private/AssetsImportController.cpp b/Engine/Plugins/Bridge/Source/MegascansPlugin/Private/AssetsImportController.cpp
--- a/Engine/Plugins/Bridge/Source/MegascansPlugin/Private/AssetsImportController.cpp
+++ b/Engine/Plugins/Bridge/Source/MegascansPlugin/Private/AssetsImportController.cpp
@@ -56,6 +56,7 @@ void FAssetsImportController::DataReceived(const FString DataFromBridge)
 			if (ContinueImport == EAppReturnType::Cancel) return;
 		}
 	}
+
 	
 	for (TSharedPtr<FJsonValue> AssetJson : AssetsImportDataArray)
 	{
@@ -73,21 +74,21 @@ void FAssetsImportController::DataReceived(const FString DataFromBridge)
 			}
 			else if (ExportMode == TEXT("progressive") && ( AssetType == TEXT("3d") || AssetType == TEXT("3dplant") ))
 			{
-				FImportProgressive3D::Get()->SetLocationOffset(LocationOffset);
-				FImportProgressive3D::Get()->ImportAsset(AssetJson->AsObject());
+				
+				FImportProgressive3D::Get()->ImportAsset(AssetJson->AsObject(), LocationOffset);
 				if (AssetJson->AsObject()->GetIntegerField(TEXT("progressiveStage")) == 1)
 				{
-					LocationOffset += 75;
+					LocationOffset += 200;
 				}
 			}
 
-			else if (ExportMode == TEXT("progressive") && AssetType == TEXT("surface") )
+			else if (ExportMode == TEXT("progressive") && (AssetType == TEXT("surface") || AssetType == TEXT("atlas")))
 			{
 				
-				FImportProgressiveSurfaces::Get()->ImportAsset(AssetJson->AsObject());
+				FImportProgressiveSurfaces::Get()->ImportAsset(AssetJson->AsObject(), LocationOffset);
 				if (AssetJson->AsObject()->GetIntegerField(TEXT("progressiveStage")) == 1)
 				{
-					LocationOffset += 75;
+					LocationOffset += 200;
 				}
 			}
 
diff --git a/Engine/Plugins/Bridge/Source/MegascansPlugin/Private/Utilities/MiscUtils.cpp b/Engine/Plugins/Bridge/Source/MegascansPlugin/Private/Utilities/MiscUtils.cpp
--- a/Engine/Plugins/Bridge/Source/MegascansPlugin/Private/Utilities/MiscUtils.cpp
+++ b/Engine/Plugins/Bridge/Source/MegascansPlugin/Private/Utilities/MiscUtils.cpp
@@ -298,10 +298,10 @@ void CopyMSPresets()
 			
 		}
 
-		if (AssetMetaData.assetType == TEXT("3dplant") && MegascansSettings->bCreateFoliage)
+		/*if (AssetMetaData.assetType == TEXT("3dplant") && MegascansSettings->bCreateFoliage)
 		{
 			AssetUtils::AddFoliageTypesToLevel(AssetMetaData.foliageAssetPaths);
-		}
+		}*/
 
 	}
 
