Adding Bridge menu item back to the workspace menu, as well as updating the icon.

diff --git a/Engine/Plugins/Bridge/Source/Bridge/Private/UI/BridgeStyle.cpp b/Engine/Plugins/Bridge/Source/Bridge/Private/UI/BridgeStyle.cpp
--- a/Engine/Plugins/Bridge/Source/Bridge/Private/UI/BridgeStyle.cpp
+++ b/Engine/Plugins/Bridge/Source/Bridge/Private/UI/BridgeStyle.cpp
@@ -39,6 +39,7 @@ FName FBridgeStyle::GetContextName()
 }
 
 #define IMAGE_BRUSH( RelativePath, ... ) FSlateImageBrush( Style->RootToContentDir( RelativePath, TEXT(".png") ), __VA_ARGS__ )
+#define IMAGE_BRUSH_SVG( RelativePath, ... ) FSlateVectorImageBrush( Style->RootToContentDir(RelativePath, TEXT(".svg")), __VA_ARGS__)
 const FVector2D Icon20x20(20.0f, 20.0f);
 const FVector2D Icon40x40(40.0f, 40.0f);
 
@@ -63,7 +64,22 @@ void FBridgeStyle::SetIcon(const FString& StyleName, const FString& ResourcePath
 	FSlateApplication::Get().GetRenderer()->ReloadTextureResources();
 }
 
+void FBridgeStyle::SetSVGIcon(const FString& StyleName, const FString& ResourcePath)
+{
+	FSlateStyleSet* Style = MSStyleInstance.Get();
+
+	FString Name(GetContextName().ToString());
+	Name = Name + "." + StyleName;
+	Style->Set(*Name, new IMAGE_BRUSH_SVG(ResourcePath, Icon40x40));
+
+	Name += ".Small";
+	Style->Set(*Name, new IMAGE_BRUSH_SVG(ResourcePath, Icon20x20));
+
+	FSlateApplication::Get().GetRenderer()->ReloadTextureResources();
+}
+
 #undef IMAGE_BRUSH
+#undef IMAGE_BRUSH_SVG
 
 const ISlateStyle& FBridgeStyle::Get()
 {
diff --git a/Engine/Plugins/Bridge/Source/Bridge/Private/UI/BridgeStyle.h b/Engine/Plugins/Bridge/Source/Bridge/Private/UI/BridgeStyle.h
--- a/Engine/Plugins/Bridge/Source/Bridge/Private/UI/BridgeStyle.h
+++ b/Engine/Plugins/Bridge/Source/Bridge/Private/UI/BridgeStyle.h
@@ -15,7 +15,7 @@ public:
 	static FName GetContextName();
 
 	static void SetIcon(const FString& StyleName, const FString& ResourcePath);
-
+	static void SetSVGIcon(const FString& StyleName, const FString& ResourcePath);
 private:
 	static TUniquePtr<FSlateStyleSet> Create();
 	static TUniquePtr<FSlateStyleSet> MSStyleInstance;
diff --git a/Engine/Plugins/Bridge/Source/Bridge/Private/UI/BridgeUIManager.cpp b/Engine/Plugins/Bridge/Source/Bridge/Private/UI/BridgeUIManager.cpp
--- a/Engine/Plugins/Bridge/Source/Bridge/Private/UI/BridgeUIManager.cpp
+++ b/Engine/Plugins/Bridge/Source/Bridge/Private/UI/BridgeUIManager.cpp
@@ -52,6 +52,7 @@ void FBridgeUIManagerImpl::SetupMenuItem()
 {
 	FBridgeStyle::SetIcon("Logo", "Logo80x80");
 	FBridgeStyle::SetIcon("ContextLogo", "Logo32x32");
+	FBridgeStyle::SetSVGIcon("MenuLogo", "QuixelBridgeB");
 	FLevelEditorModule& LevelEditorModule = FModuleManager::LoadModuleChecked<FLevelEditorModule>(LEVELEDITOR_MODULE_NAME);
 	TSharedPtr<FExtender> ToolbarExtender = MakeShareable(new FExtender);
 	ToolbarExtender->AddToolBarExtension("Settings", EExtensionHook::After, nullptr, FToolBarExtensionDelegate::CreateRaw(this, &FBridgeUIManagerImpl::FillToolbar));
@@ -74,10 +75,22 @@ void FBridgeUIManagerImpl::SetupMenuItem()
 	Section.AddMenuEntry("OpenBridgeTab",
 		LOCTEXT("OpenBridgeTab_Label", "Quixel Bridge"),
 		LOCTEXT("OpenBridgeTab_Desc", "Opens the Quixel Bridge."),
-		FSlateIcon(FBridgeStyle::GetStyleSetName(), "Bridge.Logo"),
+		FSlateIcon(FBridgeStyle::GetStyleSetName(), "Bridge.MenuLogo"),
+		FUIAction(FExecuteAction::CreateRaw(this, &FBridgeUIManagerImpl::CreateWIndow), FCanExecuteAction())
+	);
+
+	UToolMenu* WindowMenu = UToolMenus::Get()->ExtendMenu("MainFrame.MainMenu.Window");
+	FToolMenuSection* ContentSectionPtr = WindowMenu->FindSection("GetContent");
+	if (!ContentSectionPtr)
+	{
+		ContentSectionPtr = &WindowMenu->AddSection("GetContent", NSLOCTEXT("MainAppMenu", "GetContentHeader", "Get Content"));
+	}
+	ContentSectionPtr->AddMenuEntry("OpenBridgeTab",
+		LOCTEXT("OpenBridgeTab_Label", "Quixel Bridge"),
+		LOCTEXT("OpenBridgeTab_Desc", "Opens the Quixel Bridge."),
+		FSlateIcon(FBridgeStyle::GetStyleSetName(), "Bridge.MenuLogo"),
 		FUIAction(FExecuteAction::CreateRaw(this, &FBridgeUIManagerImpl::CreateWIndow), FCanExecuteAction())
 	);
-	//Section.AddSeparator(NAME_None);
 
 	//Adding Bridge entry to Content Browser context and New menu.
 	UToolMenu* ContextMenu = UToolMenus::Get()->ExtendMenu("ContentBrowser.AddNewContextMenu");
@@ -88,16 +101,11 @@ void FBridgeUIManagerImpl::SetupMenuItem()
 		"GetMegascans",
 		LOCTEXT("OpenBridgeTabText", "Add Quixel Content"),
 		LOCTEXT("GetBridgeTooltip", "Add Megascans and DHI assets to project."),
-		FSlateIcon(FBridgeStyle::GetStyleSetName(), "Bridge.Logo"),
+		FSlateIcon(FBridgeStyle::GetStyleSetName(), "Bridge.MenuLogo"),
 		FUIAction(FExecuteAction::CreateRaw(this, &FBridgeUIManagerImpl::CreateWIndow), FCanExecuteAction())
 	);
 
-	/*TSharedPtr<FExtender> NewMenuExtender = MakeShareable(new FExtender);
-	NewMenuExtender->AddMenuExtension("LevelEditor",
-		EExtensionHook::After,
-		NULL,
-		FMenuExtensionDelegate::CreateRaw(this, &FBridgeUIManagerImpl::AddPluginMenu));
-	LevelEditorModule.GetMenuExtensibilityManager()->AddExtender(NewMenuExtender);*/
+
 }
 
 void FBridgeUIManagerImpl::AddPluginMenu(FMenuBuilder& MenuBuilder)
diff --git a/Engine/Source/Editor/EditorStyle/Private/StarshipStyle.cpp b/Engine/Source/Editor/EditorStyle/Private/StarshipStyle.cpp
--- a/Engine/Source/Editor/EditorStyle/Private/StarshipStyle.cpp
+++ b/Engine/Source/Editor/EditorStyle/Private/StarshipStyle.cpp
@@ -1295,6 +1295,7 @@ void FStarshipEditorStyle::FStyle::SetupGeneralStyles()
 
 		Set( "MainFrame.ChooseFilesToSave",       new IMAGE_BRUSH_SVG( "Starship/Common/SaveChoose", Icon16x16 ) );
 		Set( "MainFrame.ConnectToSourceControl",  new IMAGE_BRUSH_SVG( "Starship/Common/SourceControl", Icon16x16 ) );
+		Set( "MainFrame.OpenMarketplace",			new IMAGE_BRUSH_SVG("Starship/MainToolbar/marketplace", Icon16x16));
 
 		Set( "MainFrame.DebugTools.SmallFont", DEFAULT_FONT( "Regular", 8 ) );
 		Set( "MainFrame.DebugTools.NormalFont", DEFAULT_FONT( "Regular", 9 ) );
diff --git a/Engine/Source/Editor/LevelEditor/Private/LevelEditorToolBar.cpp b/Engine/Source/Editor/LevelEditor/Private/LevelEditorToolBar.cpp
--- a/Engine/Source/Editor/LevelEditor/Private/LevelEditorToolBar.cpp
+++ b/Engine/Source/Editor/LevelEditor/Private/LevelEditorToolBar.cpp
@@ -1303,16 +1303,6 @@ void FLevelEditorToolBar::RegisterLevelEditorToolBar( const TSharedRef<FUIComman
 		}));
 	}
 
-	UToolMenu* SettingsToolBar = UToolMenus::Get()->RegisterMenu("LevelEditor.LevelEditorToolBar.MarketplaceToolBar", NAME_None, EMultiBoxType::SlimHorizontalToolBar);
-	SettingsToolBar->StyleName = "AssetEditorToolbar";
-	{
-		FToolMenuSection& SettingsSection = SettingsToolBar->AddSection("Settings");
-		if (FLauncherPlatformModule::Get()->CanOpenLauncher(true))
-		{
-			SettingsSection.AddEntry(FToolMenuEntry::InitToolBarButton(FLevelEditorCommands::Get().OpenMarketplace, LOCTEXT("Marketplace_Override", "Marketplace"), TAttribute<FText>(), TAttribute<FSlateIcon>(), "LevelToolbarMarketplace"));
-		}
-	}
-
 	UToolMenu* PlayToolBar = UToolMenus::Get()->RegisterMenu("LevelEditor.LevelEditorToolBar.PlayToolBar", NAME_None, EMultiBoxType::SlimHorizontalToolBar);
 	PlayToolBar->StyleName = "AssetEditorToolbar";
 	{
@@ -1881,7 +1871,10 @@ void FLevelEditorToolBar::RegisterContentMenu()
 	{
 		FToolMenuSection& Section = ContentMenu->FindOrAddSection("ExternalContent");
 		Section.InitSection("ExternalContent", LOCTEXT("ExternalContent_Label", "Get Content"), FToolMenuInsert());
-		Section.AddMenuEntry(FLevelEditorCommands::Get().OpenMarketplace);
+		if (FLauncherPlatformModule::Get()->CanOpenLauncher(true))
+		{
+			Section.AddMenuEntry(FLevelEditorCommands::Get().OpenMarketplace);
+		}
 	}
 #undef LOCTEXT_NAMESPACE
 }
diff --git a/Engine/Source/Editor/MainFrame/Private/Frame/MainFrameActions.cpp b/Engine/Source/Editor/MainFrame/Private/Frame/MainFrameActions.cpp
--- a/Engine/Source/Editor/MainFrame/Private/Frame/MainFrameActions.cpp
+++ b/Engine/Source/Editor/MainFrame/Private/Frame/MainFrameActions.cpp
@@ -172,6 +172,9 @@ void FMainFrameCommands::RegisterCommands()
 												FCanExecuteAction(),
 												FIsActionChecked::CreateStatic( &FMainFrameActionCallbacks::OpenSlateApp_IsChecked, FName("SessionFrontend" ) ) );
 
+	UI_COMMAND(OpenMarketplace, "Open Marketplace", "Opens the Marketplace", EUserInterfaceActionType::Button, FInputChord());
+	ActionList->MapAction(OpenMarketplace, FExecuteAction::CreateStatic(&FMainFrameActionCallbacks::OpenMarketplace));
+
 	UI_COMMAND(VisitOnlineLearning, "Online Learning...", "Learn Unreal Engine for free with easy-to-follow video courses and guided learning paths.", EUserInterfaceActionType::Button, FInputChord());
 	ActionList->MapAction(VisitOnlineLearning, FExecuteAction::CreateStatic(&FMainFrameActionCallbacks::VisitOnlineLearning));
 
@@ -309,6 +312,11 @@ void FMainFrameActionCallbacks::CacheProjectNames()
 	RecentProjects = GetDefault<UEditorSettings>()->RecentlyOpenedProjectFiles;
 }
 
+void FMainFrameActionCallbacks::OpenMarketplace()
+{
+	FUnrealEdMisc::Get().OpenMarketplace();
+}
+
 void FMainFrameActionCallbacks::NewProject( bool bAllowProjectOpening, bool bAllowProjectCreate )
 {
 	if (GUnrealEd->WarnIfLightingBuildIsCurrentlyRunning())
diff --git a/Engine/Source/Editor/MainFrame/Private/Frame/MainFrameActions.h b/Engine/Source/Editor/MainFrame/Private/Frame/MainFrameActions.h
--- a/Engine/Source/Editor/MainFrame/Private/Frame/MainFrameActions.h
+++ b/Engine/Source/Editor/MainFrame/Private/Frame/MainFrameActions.h
@@ -49,6 +49,7 @@ public:
 	TSharedPtr< FUICommandInfo > OpenDebugView;
 	TSharedPtr< FUICommandInfo > OpenClassViewer;
 	TSharedPtr< FUICommandInfo > OpenWidgetReflector;
+	TSharedPtr< FUICommandInfo > OpenMarketplace;
 
 	TSharedPtr< FUICommandInfo > VisitOnlineLearning;
 	TSharedPtr< FUICommandInfo > VisitForums;
@@ -231,7 +232,8 @@ public:
 	/** Gathers all available projects the user can switch to from main menu */
 	static void CacheProjectNames();
 
-
+	/** Opens the marketplace */
+	static void OpenMarketplace();
 public:
 
 	// List of projects that the user can switch to.
diff --git a/Engine/Source/Editor/MainFrame/Private/Menus/MainMenu.cpp b/Engine/Source/Editor/MainFrame/Private/Menus/MainMenu.cpp
--- a/Engine/Source/Editor/MainFrame/Private/Menus/MainMenu.cpp
+++ b/Engine/Source/Editor/MainFrame/Private/Menus/MainMenu.cpp
@@ -29,6 +29,7 @@
 #include "UndoHistoryModule.h"
 #include "Framework/Commands/GenericCommands.h"
 #include "ITranslationEditor.h"
+#include "LauncherPlatformModule.h"
 
 #define LOCTEXT_NAMESPACE "MainFileMenu"
 
@@ -214,6 +215,17 @@ void FMainMenu::RegisterWindowMenu()
 		}
 	}));
 
+	// Get content section
+	{
+		if (FLauncherPlatformModule::Get()->CanOpenLauncher(true))
+		{
+			FToolMenuSection& Section = Menu->AddSection("GetContent", NSLOCTEXT("MainAppMenu", "GetContentHeader", "Get Content"));
+			Section.AddMenuEntry(FMainFrameCommands::Get().OpenMarketplace);
+		}
+		
+	}
+
+
 	// Layout section
 	{
 		FToolMenuSection& Section = Menu->AddSection("WindowLayout", NSLOCTEXT("MainAppMenu", "LayoutManagementHeader", "Layout"));
