- If CEF fails to load then don't crash in the Quixel bridge UI

diff --git a/Engine/Plugins/Bridge/Source/Bridge/Private/UI/BridgeUIManager.cpp b/Engine/Plugins/Bridge/Source/Bridge/Private/UI/BridgeUIManager.cpp
--- a/Engine/Plugins/Bridge/Source/Bridge/Private/UI/BridgeUIManager.cpp
+++ b/Engine/Plugins/Bridge/Source/Bridge/Private/UI/BridgeUIManager.cpp
@@ -168,12 +168,16 @@ void FBridgeUIManagerImpl::CreateWIndow()
 
 	WindowSettings.OSWindowHandle = BridgeWindow.Get()->GetNativeWindow().Get()->GetOSWindowHandle();
 
-	TSharedPtr<IWebBrowserWindow> Browser = IWebBrowserModule::Get().GetSingleton()->CreateBrowserWindow(WindowSettings);
+	// IWebBrowserModule::Get() was already callled in WebBrowserWidgetModule.cpp so we don't need to force the load again here
+	if (IWebBrowserModule::IsAvailable() && IWebBrowserModule::Get().IsWebModuleAvailable())
+	{
+		TSharedPtr<IWebBrowserWindow> Browser = IWebBrowserModule::Get().GetSingleton()->CreateBrowserWindow(WindowSettings);
 
-	BridgeWindow.Get()->SetContent(
-		SAssignNew(WebBrowserWidget, SWebBrowser, Browser)
-		.ShowControls(false)
-	);
+		BridgeWindow.Get()->SetContent(
+			SAssignNew(WebBrowserWidget, SWebBrowser, Browser)
+			.ShowControls(false)
+		);
+	}
 
 	// Create and add a transparent overlay window to enable drag n drop on top of CEF
 	FBridgeUIManager::Instance->OverlayWindow = SNew(SWindow)
@@ -209,16 +213,25 @@ void FBridgeUIManagerImpl::CreateWIndow()
 	BridgeWindow.Get()->SetOnWindowClosed(
 		FOnWindowClosed::CreateLambda([](const TSharedRef<SWindow>& Window)
 		{
-			FBridgeUIManager::Instance->WebBrowserWidget.Reset();
-			FBridgeUIManager::BrowserBinding->OnExitDelegate.ExecuteIfBound(TEXT("test"));
+			if (FBridgeUIManager::Instance->WebBrowserWidget.IsValid())
+			{
+				FBridgeUIManager::Instance->WebBrowserWidget.Reset();
+			}
+			if (FBridgeUIManager::BrowserBinding != nullptr)
+			{
+				FBridgeUIManager::BrowserBinding->OnExitDelegate.ExecuteIfBound(TEXT("test"));
+			}
 			FBridgeUIManager::Instance->BridgeWindow = NULL;
 		})
 	);
 
-	UNodePort* NodePortInfo = NewObject<UNodePort>();
-	FBridgeUIManager::BrowserBinding = NewObject<UBrowserBinding>();
-	WebBrowserWidget->BindUObject(TEXT("NodePortInfo"), NodePortInfo, true);
-	WebBrowserWidget->BindUObject(TEXT("BrowserBinding"), FBridgeUIManager::BrowserBinding, true);
+	if (WebBrowserWidget.IsValid())
+	{
+		UNodePort* NodePortInfo = NewObject<UNodePort>();
+		FBridgeUIManager::BrowserBinding = NewObject<UBrowserBinding>();
+		WebBrowserWidget->BindUObject(TEXT("NodePortInfo"), NodePortInfo, true);
+		WebBrowserWidget->BindUObject(TEXT("BrowserBinding"), FBridgeUIManager::BrowserBinding, true);
+	}
 }
 
 void FBridgeUIManager::Shutdown()
