- When creating the bridge plugin force the webbrowser plugin to load as it uses it

diff --git a/Engine/Plugins/Bridge/Source/Bridge/Private/UI/BridgeUIManager.cpp b/Engine/Plugins/Bridge/Source/Bridge/Private/UI/BridgeUIManager.cpp
--- a/Engine/Plugins/Bridge/Source/Bridge/Private/UI/BridgeUIManager.cpp
+++ b/Engine/Plugins/Bridge/Source/Bridge/Private/UI/BridgeUIManager.cpp
@@ -168,7 +168,7 @@ void FBridgeUIManagerImpl::CreateWIndow()
 
 	WindowSettings.OSWindowHandle = BridgeWindow.Get()->GetNativeWindow().Get()->GetOSWindowHandle();
 
-	// IWebBrowserModule::Get() was already callled in WebBrowserWidgetModule.cpp so we don't need to force the load again here
+	IWebBrowserModule::Get(); // force the module to load, if the project doesn't otherwise use the browser plugin
 	if (IWebBrowserModule::IsAvailable() && IWebBrowserModule::Get().IsWebModuleAvailable())
 	{
 		TSharedPtr<IWebBrowserWindow> Browser = IWebBrowserModule::Get().GetSingleton()->CreateBrowserWindow(WindowSettings);
