Fix Bridge Plugin error Int8Property assetTier is not initialized properly

diff --git a/Engine/Plugins/Bridge/Source/MegascansPlugin/Private/MSAssetImportData.h b/Engine/Plugins/Bridge/Source/MegascansPlugin/Private/MSAssetImportData.h
--- a/Engine/Plugins/Bridge/Source/MegascansPlugin/Private/MSAssetImportData.h
+++ b/Engine/Plugins/Bridge/Source/MegascansPlugin/Private/MSAssetImportData.h
@@ -170,7 +170,7 @@ public:
 	UPROPERTY()
 		FString assetSubType;
 	UPROPERTY()
-		int8 assetTier;
+		int8 assetTier = 0;
 	UPROPERTY()
 		FString assetRootPath;
 	UPROPERTY()
