Add Bridge menu item to the workspace menu, as well as update the icon

diff --git a/Engine/Plugins/Bridge/Source/Bridge/Private/UI/BridgeStyle.cpp b/Engine/Plugins/Bridge/Source/Bridge/Private/UI/BridgeStyle.cpp
--- a/Engine/Plugins/Bridge/Source/Bridge/Private/UI/BridgeStyle.cpp
+++ b/Engine/Plugins/Bridge/Source/Bridge/Private/UI/BridgeStyle.cpp
@@ -39,6 +39,7 @@ FName FBridgeStyle::GetContextName()
 }
 
 #define IMAGE_BRUSH( RelativePath, ... ) FSlateImageBrush( Style->RootToContentDir( RelativePath, TEXT(".png") ), __VA_ARGS__ )
+#define IMAGE_BRUSH_SVG( RelativePath, ... ) FSlateVectorImageBrush( Style->RootToContentDir(RelativePath, TEXT(".svg")), __VA_ARGS__)
 const FVector2D Icon20x20(20.0f, 20.0f);
 const FVector2D Icon40x40(40.0f, 40.0f);
 
@@ -63,7 +64,22 @@ void FBridgeStyle::SetIcon(const FString& StyleName, const FString& ResourcePath
 	FSlateApplication::Get().GetRenderer()->ReloadTextureResources();
 }
 
+void FBridgeStyle::SetSVGIcon(const FString& StyleName, const FString& ResourcePath)
+{
+	FSlateStyleSet* Style = MSStyleInstance.Get();
+
+	FString Name(GetContextName().ToString());
+	Name = Name + "." + StyleName;
+	Style->Set(*Name, new IMAGE_BRUSH_SVG(ResourcePath, Icon40x40));
+
+	Name += ".Small";
+	Style->Set(*Name, new IMAGE_BRUSH_SVG(ResourcePath, Icon20x20));
+
+	FSlateApplication::Get().GetRenderer()->ReloadTextureResources();
+}
+
 #undef IMAGE_BRUSH
+#undef IMAGE_BRUSH_SVG
 
 const ISlateStyle& FBridgeStyle::Get()
 {
diff --git a/Engine/Plugins/Bridge/Source/Bridge/Private/UI/BridgeStyle.h b/Engine/Plugins/Bridge/Source/Bridge/Private/UI/BridgeStyle.h
--- a/Engine/Plugins/Bridge/Source/Bridge/Private/UI/BridgeStyle.h
+++ b/Engine/Plugins/Bridge/Source/Bridge/Private/UI/BridgeStyle.h
@@ -15,7 +15,7 @@ public:
 	static FName GetContextName();
 
 	static void SetIcon(const FString& StyleName, const FString& ResourcePath);
-
+	static void SetSVGIcon(const FString& StyleName, const FString& ResourcePath);
 private:
 	static TUniquePtr<FSlateStyleSet> Create();
 	static TUniquePtr<FSlateStyleSet> MSStyleInstance;
diff --git a/Engine/Plugins/Bridge/Source/Bridge/Private/UI/BridgeUIManager.cpp b/Engine/Plugins/Bridge/Source/Bridge/Private/UI/BridgeUIManager.cpp
--- a/Engine/Plugins/Bridge/Source/Bridge/Private/UI/BridgeUIManager.cpp
+++ b/Engine/Plugins/Bridge/Source/Bridge/Private/UI/BridgeUIManager.cpp
@@ -52,6 +52,7 @@ void FBridgeUIManagerImpl::SetupMenuItem()
 {
 	FBridgeStyle::SetIcon("Logo", "Logo80x80");
 	FBridgeStyle::SetIcon("ContextLogo", "Logo32x32");
+	FBridgeStyle::SetSVGIcon("MenuLogo", "QuixelBridgeB");
 	FLevelEditorModule& LevelEditorModule = FModuleManager::LoadModuleChecked<FLevelEditorModule>(LEVELEDITOR_MODULE_NAME);
 	TSharedPtr<FExtender> ToolbarExtender = MakeShareable(new FExtender);
 	ToolbarExtender->AddToolBarExtension("Settings", EExtensionHook::After, nullptr, FToolBarExtensionDelegate::CreateRaw(this, &FBridgeUIManagerImpl::FillToolbar));
@@ -74,7 +75,20 @@ void FBridgeUIManagerImpl::SetupMenuItem()
 	Section.AddMenuEntry("OpenBridgeTab",
 		LOCTEXT("OpenBridgeTab_Label", "Quixel Bridge"),
 		LOCTEXT("OpenBridgeTab_Desc", "Opens the Quixel Bridge."),
-		FSlateIcon(FBridgeStyle::GetStyleSetName(), "Bridge.Logo"),
+		FSlateIcon(FBridgeStyle::GetStyleSetName(), "Bridge.MenuLogo"),
+		FUIAction(FExecuteAction::CreateRaw(this, &FBridgeUIManagerImpl::CreateWindow), FCanExecuteAction())
+	);
+
+	UToolMenu* WindowMenu = UToolMenus::Get()->ExtendMenu("MainFrame.MainMenu.Window");
+	FToolMenuSection* ContentSectionPtr = WindowMenu->FindSection("GetContent");
+	if (!ContentSectionPtr)
+	{
+		ContentSectionPtr = &WindowMenu->AddSection("GetContent", NSLOCTEXT("MainAppMenu", "GetContentHeader", "Get Content"));
+	}
+	ContentSectionPtr->AddMenuEntry("OpenBridgeTab",
+		LOCTEXT("OpenBridgeTab_Label", "Quixel Bridge"),
+		LOCTEXT("OpenBridgeTab_Desc", "Opens the Quixel Bridge."),
+		FSlateIcon(FBridgeStyle::GetStyleSetName(), "Bridge.MenuLogo"),
 		FUIAction(FExecuteAction::CreateRaw(this, &FBridgeUIManagerImpl::CreateWindow), FCanExecuteAction())
 	);
 	//Section.AddSeparator(NAME_None);
@@ -88,7 +102,7 @@ void FBridgeUIManagerImpl::SetupMenuItem()
 		"GetMegascans",
 		LOCTEXT("OpenBridgeTabText", "Add Quixel Content"),
 		LOCTEXT("GetBridgeTooltip", "Add Megascans and DHI assets to project."),
-		FSlateIcon(FBridgeStyle::GetStyleSetName(), "Bridge.Logo"),
+		FSlateIcon(FBridgeStyle::GetStyleSetName(), "Bridge.MenuLogo"),
 		FUIAction(FExecuteAction::CreateRaw(this, &FBridgeUIManagerImpl::CreateWindow), FCanExecuteAction())
 	);
 
