Avoid Engine Version for Bridge plugin

diff --git a/Engine/Plugins/Bridge/Bridge.uplugin b/Engine/Plugins/Bridge/Bridge.uplugin
--- a/Engine/Plugins/Bridge/Bridge.uplugin
+++ b/Engine/Plugins/Bridge/Bridge.uplugin
@@ -10,7 +10,6 @@
 	"DocsURL": "https://help.quixel.com/hc/en-us/sections/360005846137-Quixel-Bridge-for-Unreal-Engine-5",
 	"MarketplaceURL": "",
 	"SupportURL": "[REDACTED_EMAIL]",
-	"EngineVersion": "5.0.0",
 	"CanContainContent": true,
 	"EnabledByDefault": true,
 	"SupportedTargetPlatforms": [
@@ -46,4 +45,4 @@
 			"Enabled": true
 		}
 	]
-}
\ No newline at end of file
+}
