Fix NonUnity compilation errors generated from the Bridge plugin

diff --git a/Engine/Plugins/Bridge/Source/MegascansPlugin/Public/BridgeDragDropUtils.h b/Engine/Plugins/Bridge/Source/MegascansPlugin/Public/BridgeDragDropUtils.h
--- a/Engine/Plugins/Bridge/Source/MegascansPlugin/Public/BridgeDragDropUtils.h
+++ b/Engine/Plugins/Bridge/Source/MegascansPlugin/Public/BridgeDragDropUtils.h
@@ -4,9 +4,9 @@
 #include "CoreMinimal.h"
 #include "AssetRegistry/AssetData.h"
 
-DECLARE_DELEGATE_ThreeParams(FOnAddProgressiveStageDataCallback, FAssetData AssetData, FString AssetId, AStaticMeshActor* SpawnedActor);
+class AStaticMeshActor;
 
-class FBridgeDragDropImpl;
+DECLARE_DELEGATE_ThreeParams(FOnAddProgressiveStageDataCallback, FAssetData AssetData, FString AssetId, AStaticMeshActor* SpawnedActor);
 
 class MEGASCANSPLUGIN_API FBridgeDragDropImpl : public TSharedFromThis<FBridgeDragDropImpl>
 {
