Fix crash on Bridge plugin launch (MacOS)

diff --git a/Engine/Plugins/Bridge/Source/Bridge/Private/UI/BridgeUIManager.cpp b/Engine/Plugins/Bridge/Source/Bridge/Private/UI/BridgeUIManager.cpp
--- a/Engine/Plugins/Bridge/Source/Bridge/Private/UI/BridgeUIManager.cpp
+++ b/Engine/Plugins/Bridge/Source/Bridge/Private/UI/BridgeUIManager.cpp
@@ -250,7 +250,14 @@ TSharedRef<SDockTab> FBridgeUIManagerImpl::CreateBridgeTab(const FSpawnTabArgs&
 
 	IWebBrowserSingleton* WebBrowserSingleton = IWebBrowserModule::Get().GetSingleton();
 
+#if PLATFORM_MAC
+	// On Mac we need to ignore IWebBrowserModule::Get() call as it causes a crash if Web Browser widget isn't enabled
+	// The cause has to do with internal CEF implementation
+	// Proper fix will come in the form of CEF upgrade to v90 (https://jira.it.epicgames.com/browse/DISTRO-2434)
+	if (IWebBrowserModule::IsAvailable())
+#else
 	if (IWebBrowserModule::IsAvailable() && IWebBrowserModule::Get().IsWebModuleAvailable())
+#endif
 	{
 		Browser = WebBrowserSingleton->CreateBrowserWindow(WindowSettings);
 		SAssignNew(LocalBrowserDock, SDockTab)
