Fix MacOS crash on Bridge plugin launch when shaders are compiling

Added WebBrowserWidget as an explicit plugin dependency in Bridge

diff --git a/Engine/Plugins/Bridge/Bridge.uplugin b/Engine/Plugins/Bridge/Bridge.uplugin
--- a/Engine/Plugins/Bridge/Bridge.uplugin
+++ b/Engine/Plugins/Bridge/Bridge.uplugin
@@ -43,6 +43,10 @@
 		{
 			"Name": "EditorScriptingUtilities",
 			"Enabled": true
+		},
+		{
+			"Name": "WebBrowserWidget",
+			"Enabled": true
 		}
 	]
 }
