Backing out CL 18984951 as its causing many issues, including no longer being able to report crashes for Linux/Mac

diff --git a/Engine/Plugins/Bridge/Bridge.uplugin b/Engine/Plugins/Bridge/Bridge.uplugin
--- a/Engine/Plugins/Bridge/Bridge.uplugin
+++ b/Engine/Plugins/Bridge/Bridge.uplugin
@@ -43,10 +43,6 @@
 		{
 			"Name": "EditorScriptingUtilities",
 			"Enabled": true
-		},
-		{
-			"Name": "WebBrowserWidget",
-			"Enabled": true
 		}
 	]
 }
