Temp workaround for fixing crash on Bridge plugin launch (Mac OS)

Settled on showing a popup on Mac, that asks the users to enable Web Browser plugin for using Bridge.

diff --git a/Engine/Plugins/Bridge/Source/Bridge/Bridge.Build.cs b/Engine/Plugins/Bridge/Source/Bridge/Bridge.Build.cs
--- a/Engine/Plugins/Bridge/Source/Bridge/Bridge.Build.cs
+++ b/Engine/Plugins/Bridge/Source/Bridge/Bridge.Build.cs
@@ -25,6 +25,7 @@ namespace UnrealBuildTool.Rules
 		{
 			PublicDependencyModuleNames.AddRange(
 				new string[] {
+					"Projects",
 					"Core",
 					"CoreUObject",
 					"UnrealEd",
diff --git a/Engine/Plugins/Bridge/Source/Bridge/Private/UI/BridgeUIManager.cpp b/Engine/Plugins/Bridge/Source/Bridge/Private/UI/BridgeUIManager.cpp
--- a/Engine/Plugins/Bridge/Source/Bridge/Private/UI/BridgeUIManager.cpp
+++ b/Engine/Plugins/Bridge/Source/Bridge/Private/UI/BridgeUIManager.cpp
@@ -8,6 +8,7 @@
 #include "NodeProcess.h"
 #include "Serialization/JsonReader.h"
 #include "JsonObjectConverter.h"
+#include "Interfaces/IPluginManager.h"
 
 // WebBrowser
 #include "SWebBrowser.h"
@@ -165,6 +166,16 @@ void FBridgeUIManagerImpl::FillToolbar(FToolBarBuilder& ToolbarBuilder)
 
 void FBridgeUIManagerImpl::CreateWindow()
 {
+#if PLATFORM_MAC
+	// Check if WebBrowserWidget plugin is enabled
+	TSharedPtr<IPlugin> WebBrowserPlugin = IPluginManager::Get().FindPlugin("WebBrowserWidget");
+	if (WebBrowserPlugin.IsValid() && !WebBrowserPlugin->IsEnabled())
+	{
+		FMessageDialog::Open(EAppMsgType::Ok, LOCTEXT("WebBrowserWidgetPluginNotEnabled", "Web Browser plugin is not enabled. Please enable it in the plugin manager to use Bridge."));
+		return;
+	}
+#endif
+
 	FGlobalTabmanager::Get()->TryInvokeTab(BridgeTabName);
 
 	// Set desired window size (if the desired window size is less than main window size)
@@ -208,16 +219,25 @@ void FBridgeUIManager::Shutdown()
 
 TSharedRef<SDockTab> FBridgeUIManagerImpl::CreateBridgeTab(const FSpawnTabArgs& Args)
 {
-	// Start node process
-	FNodeProcessManager::Get()->StartNodeProcess();
+#if PLATFORM_MAC
+	// Check if WebBrowserWidget plugin is enabled
+	TSharedPtr<IPlugin> WebBrowserPlugin = IPluginManager::Get().FindPlugin("WebBrowserWidget");
+	if (WebBrowserPlugin.IsValid() && !WebBrowserPlugin->IsEnabled())
+	{
+		FMessageDialog::Open(EAppMsgType::Ok, LOCTEXT("WebBrowserWidgetPluginNotEnabled", "Web Browser plugin is not enabled. Please enable it in the plugin manager to use Bridge."));
 
-	// Delay launch on Mac & Linux
-#if PLATFORM_MAC || PLATFORM_LINUX
-	FGenericPlatformProcess::Sleep(2);
-#endif
+		return SAssignNew(LocalBrowserDock, SDockTab)
+			.TabRole(ETabRole::NomadTab);
+	}
 
-	FWebBrowserInitSettings browserInitSettings = FWebBrowserInitSettings();
-	IWebBrowserModule::Get().CustomInitialize(browserInitSettings);
+	// Call to GetSingleton initializes WebBrowser
+	if (!(IWebBrowserModule::Get().IsWebModuleAvailable() && IWebBrowserModule::Get().GetSingleton() != nullptr))
+	{
+		UE_LOG(LogTemp, Verbose, TEXT("Bridge: WebBrowserModule is not available"));
+	}
+#endif
+	// Start node process
+	FNodeProcessManager::Get()->StartNodeProcess();
 
 	FString PluginPath = FPaths::Combine(FPaths::EnginePluginsDir(), TEXT("Bridge"));
 	FString IndexUrl = FPaths::ConvertRelativePathToFull(FPaths::Combine(PluginPath, TEXT("ThirdParty"), TEXT("megascans"), TEXT("index.html")));
@@ -229,6 +249,7 @@ TSharedRef<SDockTab> FBridgeUIManagerImpl::CreateBridgeTab(const FSpawnTabArgs&
 	FString Token;
 	FFileHelper::LoadFileToString(Token, *TokenPath);
 
+	FString FinalUrl;
 	if (Token.Len() > 0)
 	{
 		TSharedPtr<FJsonObject> JsonParsed;
@@ -238,68 +259,75 @@ TSharedRef<SDockTab> FBridgeUIManagerImpl::CreateBridgeTab(const FSpawnTabArgs&
 			FString AccessToken = JsonParsed->GetStringField("token");
 			FString RefreshToken = JsonParsed->GetStringField("refreshToken");
 			const FString FileUrl = FPaths::Combine(TEXT("file:///"), IndexUrl);
-			const FString FinalUrl = FString::Printf(TEXT("%s?token=%s&refreshToken=%s"), *FileUrl, *AccessToken, *RefreshToken);
-			WindowSettings.InitialURL = FinalUrl;
+			FinalUrl = FString::Printf(TEXT("%s?token=%s&refreshToken=%s"), *FileUrl, *AccessToken, *RefreshToken);
 		}
 	}
 	else
 	{
-		WindowSettings.InitialURL = FPaths::Combine(TEXT("file:///"), IndexUrl);
+		FinalUrl = FPaths::Combine(TEXT("file:///"), IndexUrl);
 	}
 
-	WindowSettings.BrowserFrameRate = 60;
-
-	IWebBrowserSingleton* WebBrowserSingleton = IWebBrowserModule::Get().GetSingleton();
+	TSharedPtr<SWebBrowser> PluginWebBrowser;
 
 #if PLATFORM_MAC
-	// On Mac we need to ignore IWebBrowserModule::Get() call as it causes a crash if Web Browser widget isn't enabled
-	// The cause has to do with internal CEF implementation
-	// Proper fix will come in the form of CEF upgrade to v90 (https://jira.it.epicgames.com/browse/DISTRO-2434)
-	if (IWebBrowserModule::IsAvailable())
-#else
+	PluginWebBrowser = SAssignNew(WebBrowserWidget, SWebBrowser)
+		.ShowAddressBar(false)
+		.ShowControls(false)
+		.InitialURL(FinalUrl);
+#elif PLATFORM_WINDOWS || PLATFORM_LINUX
+	FWebBrowserInitSettings browserInitSettings = FWebBrowserInitSettings();
+	IWebBrowserModule::Get().CustomInitialize(browserInitSettings);
+	WindowSettings.InitialURL = FinalUrl;
+	WindowSettings.BrowserFrameRate = 60;
 	if (IWebBrowserModule::IsAvailable() && IWebBrowserModule::Get().IsWebModuleAvailable())
-#endif
 	{
+		IWebBrowserSingleton* WebBrowserSingleton = IWebBrowserModule::Get().GetSingleton();
 		Browser = WebBrowserSingleton->CreateBrowserWindow(WindowSettings);
-		SAssignNew(LocalBrowserDock, SDockTab)
-			.OnTabClosed_Lambda([](TSharedRef<class SDockTab> InParentTab)
-			{
-				// Kill node process if bound
-				FBridgeUIManager::BrowserBinding->OnExitDelegate.ExecuteIfBound("Plugin Window Closed");
-				FBridgeUIManager::BrowserBinding = NULL;
-				// Adding a delay for Mac, for node process to exit completely (before plugin)
-			#if PLATFORM_MAC
-				FGenericPlatformProcess::Sleep(0.3);
-			#endif
-
-				// Clean up browser
-				FBridgeUIManager::Instance->LocalBrowserDock = NULL;
-				if (FBridgeUIManager::Instance->WebBrowserWidget.IsValid())
-				{
-					FBridgeUIManager::Instance->WebBrowserWidget.Reset();
-					FBridgeUIManager::Instance->Browser.Reset();
-				}
-			})
-			.TabRole(ETabRole::NomadTab)
-			[
-				SAssignNew(WebBrowserWidget, SWebBrowser, Browser)
-				.ShowAddressBar(false)
-				.ShowControls(false)
-			];
-
-		LocalBrowserDock->SetOnTabDraggedOverDockArea(
-			FSimpleDelegate::CreateLambda([IndexUrl]()
-										  {
-											FBridgeUIManager::Instance->WebBrowserWidget->Invalidate(EInvalidateWidgetReason::Layout);
-										  })
-		);
-		LocalBrowserDock->SetOnTabRelocated(
-			FSimpleDelegate::CreateLambda([IndexUrl]()
-										  {
-											FBridgeUIManager::Instance->WebBrowserWidget->Invalidate(EInvalidateWidgetReason::Layout);
-										  })
-		);
+		PluginWebBrowser = SAssignNew(WebBrowserWidget, SWebBrowser, Browser)
+			.ShowAddressBar(false)
+			.ShowControls(false);
 	}
+	else
+	{
+		FMessageDialog::Open(EAppMsgType::Ok, LOCTEXT("WebBrowserWidgetPluginNotEnabled", "Failed to load the plugin. Please enable Web Browser in the plugin manager to use Bridge."));
+
+		return SAssignNew(LocalBrowserDock, SDockTab)
+			.TabRole(ETabRole::NomadTab);
+	}
+#endif
+
+	SAssignNew(LocalBrowserDock, SDockTab)
+		.OnTabClosed_Lambda([](TSharedRef<class SDockTab> InParentTab)
+		{
+			// Kill node process if bound
+			FBridgeUIManager::BrowserBinding->OnExitDelegate.ExecuteIfBound("Plugin Window Closed");
+			FBridgeUIManager::BrowserBinding = NULL;
+
+			// Clean up browser
+			FBridgeUIManager::Instance->LocalBrowserDock = NULL;
+			if (FBridgeUIManager::Instance->WebBrowserWidget.IsValid())
+			{
+				FBridgeUIManager::Instance->WebBrowserWidget.Reset();
+				FBridgeUIManager::Instance->Browser.Reset();
+			}
+		})
+		.TabRole(ETabRole::NomadTab)
+		[
+			PluginWebBrowser.ToSharedRef()
+		];
+
+	LocalBrowserDock->SetOnTabDraggedOverDockArea(
+		FSimpleDelegate::CreateLambda([IndexUrl]()
+										{
+										FBridgeUIManager::Instance->WebBrowserWidget->Invalidate(EInvalidateWidgetReason::Layout);
+										})
+	);
+	LocalBrowserDock->SetOnTabRelocated(
+		FSimpleDelegate::CreateLambda([IndexUrl]()
+										{
+										FBridgeUIManager::Instance->WebBrowserWidget->Invalidate(EInvalidateWidgetReason::Layout);
+										})
+	);
 
 	if (WebBrowserWidget.IsValid())
 	{
