Fix editor getting resized on Bridge plugin launch if it was previously tabbed

diff --git a/Engine/Plugins/Bridge/Source/Bridge/Private/UI/BridgeUIManager.cpp b/Engine/Plugins/Bridge/Source/Bridge/Private/UI/BridgeUIManager.cpp
--- a/Engine/Plugins/Bridge/Source/Bridge/Private/UI/BridgeUIManager.cpp
+++ b/Engine/Plugins/Bridge/Source/Bridge/Private/UI/BridgeUIManager.cpp
@@ -189,6 +189,12 @@ void FBridgeUIManagerImpl::CreateWindow()
 
 		if (DesiredWidth < MainWindowSize.X && DesiredHeight < MainWindowSize.Y && LocalBrowserDock->GetParentWindow().IsValid())
 		{
+			// If Bridge is docked as a tab, the parent window will be the main window
+			if (LocalBrowserDock->GetParentWindow() == Windows[0])
+			{
+				return;
+			}
+
 			LocalBrowserDock->GetParentWindow()->Resize(FVector2D(DesiredWidth, DesiredHeight));
 			LocalBrowserDock->GetParentWindow()->MoveWindowTo(FVector2D((MainWindowSize.X - DesiredWidth) - 17, MainWindowSize.Y - DesiredHeight) / 2);
 		}
