ensure the RigLogic plugin is enabled and loaded before we load ControlRigs using RigLogic nodes

diff --git a/Engine/Plugins/Bridge/Source/MegascansPlugin/MegascansPlugin.Build.cs b/Engine/Plugins/Bridge/Source/MegascansPlugin/MegascansPlugin.Build.cs
--- a/Engine/Plugins/Bridge/Source/MegascansPlugin/MegascansPlugin.Build.cs
+++ b/Engine/Plugins/Bridge/Source/MegascansPlugin/MegascansPlugin.Build.cs
@@ -35,6 +35,7 @@ namespace UnrealBuildTool.Rules
 					"MeshBuilder",
 					"TargetPlatform",
 					"EditorScriptingUtilities",
+					"Projects"
 				}
 			);
 		}
diff --git a/Engine/Plugins/Bridge/Source/MegascansPlugin/Private/AssetImporters/DHIImport.cpp b/Engine/Plugins/Bridge/Source/MegascansPlugin/Private/AssetImporters/DHIImport.cpp
--- a/Engine/Plugins/Bridge/Source/MegascansPlugin/Private/AssetImporters/DHIImport.cpp
+++ b/Engine/Plugins/Bridge/Source/MegascansPlugin/Private/AssetImporters/DHIImport.cpp
@@ -8,6 +8,8 @@
 #include "EditorAssetLibrary.h"
 #include "GenericPlatform/GenericPlatformFile.h"
 #include "HAL/PlatformFileManager.h"
+#include "Interfaces/IPluginManager.h"
+#include "Interfaces/IProjectManager.h"
 #include "Misc/MessageDialog.h"
 #include "Misc/Paths.h"
 #include "Internationalization/Text.h"
@@ -141,10 +143,51 @@ void FImportDHI::ImportAsset(TSharedPtr<FJsonObject> AssetImportJson)
 	FString BPPath = FPaths::Combine(TEXT("/Game/MetaHumans/"), CharacterSourceData->CharacterName, BPName);
 
 	FAssetData CharacterAssetData = AssetRegistryModule.Get().GetAssetByObjectPath(FName(*BPPath));
+	
+	// NOTE: the RigLogic plugin (and maybe others) must be loaded and added to the project before loading the asset
+	// otherwise we get rid of the RigLogic nodes, resulting in leaving the asset in an undefined state. In the context
+	// of ControlRig assets, graphs will remove the RigLogic nodes if the plugin is not enabled because the
+	// FRigUnit_RigLogic_Data won't be available
+	EnableMissingPlugins();
+
+	// As part of our discussion to fix this, here was the explanation about why the asset has to be loaded:
+	// "In UE4 we came across this issue where even if we did call the syncfolder on content browser, it would still
+	// not show the character in content browser. So a workaround to make the character appear was to just load
+	// the character."
 	UObject* CharacterObject = CharacterAssetData.GetAsset();
 
 	AssetUtils::FocusOnSelected(CharacterDestination);
-	
+}
 
+void FImportDHI::EnableMissingPlugins()
+{
+	// TODO we should find a way to retrieve the required plugins from the metadata as RigLogic might not be the only one 
+	static const TArray<FString> NeededPluginNames( {TEXT("RigLogic")} );
 
+	IPluginManager& PluginManager = IPluginManager::Get();
+	IProjectManager& ProjectManager = IProjectManager::Get();
+	
+	for (const FString& PluginName: NeededPluginNames)
+	{
+		TSharedPtr<IPlugin> NeededPlugin = PluginManager.FindPlugin(PluginName);
+		if (NeededPlugin.IsValid() && !NeededPlugin->IsEnabled())
+		{
+			FText FailMessage;
+			bool bPluginEnabled = ProjectManager.SetPluginEnabled(NeededPlugin->GetName(), true, FailMessage);
+			
+			if (bPluginEnabled && ProjectManager.IsCurrentProjectDirty())
+			{
+				bPluginEnabled = ProjectManager.SaveCurrentProjectToDisk(FailMessage);
+			}
+
+			if (bPluginEnabled)
+			{
+				PluginManager.MountNewlyCreatedPlugin(NeededPlugin->GetName());
+			}
+			else
+			{
+				FMessageDialog::Open(EAppMsgType::Ok, FailMessage);
+			}
+		}
+	}
 }
diff --git a/Engine/Plugins/Bridge/Source/MegascansPlugin/Private/AssetImporters/DHIImport.h b/Engine/Plugins/Bridge/Source/MegascansPlugin/Private/AssetImporters/DHIImport.h
--- a/Engine/Plugins/Bridge/Source/MegascansPlugin/Private/AssetImporters/DHIImport.h
+++ b/Engine/Plugins/Bridge/Source/MegascansPlugin/Private/AssetImporters/DHIImport.h
@@ -26,4 +26,7 @@ private:
 public:
 	static TSharedPtr<FImportDHI> Get();
 	void ImportAsset(TSharedPtr<FJsonObject> AssetImportJson);
+
+private:
+	static void EnableMissingPlugins();
 };
\ No newline at end of file
