Update the Bridge message handler to have the newer version of  OnController functions.

diff --git a/Engine/Plugins/Bridge/Source/Bridge/Private/UI/FBridgeMessageHandler.cpp b/Engine/Plugins/Bridge/Source/Bridge/Private/UI/FBridgeMessageHandler.cpp
--- a/Engine/Plugins/Bridge/Source/Bridge/Private/UI/FBridgeMessageHandler.cpp
+++ b/Engine/Plugins/Bridge/Source/Bridge/Private/UI/FBridgeMessageHandler.cpp
@@ -201,6 +201,7 @@ bool FBridgeMessageHandler::OnCursorSet()
 	return TargetHandler->OnCursorSet();
 }
 
+PRAGMA_DISABLE_DEPRECATION_WARNINGS
 bool FBridgeMessageHandler::OnControllerAnalog(FGamepadKeyNames::Type KeyName, int32 ControllerId, float AnalogValue)
 {
 	return TargetHandler->OnControllerAnalog(KeyName, ControllerId, AnalogValue);
@@ -215,6 +216,22 @@ bool FBridgeMessageHandler::OnControllerButtonReleased(FGamepadKeyNames::Type Ke
 {
 	return TargetHandler->OnControllerButtonReleased(KeyName, ControllerId, IsRepeat);
 }
+PRAGMA_ENABLE_DEPRECATION_WARNINGS
+
+bool FBridgeMessageHandler::OnControllerAnalog(FGamepadKeyNames::Type KeyName, FPlatformUserId PlatformUserId, FInputDeviceId InputDeviceId, float AnalogValue)
+{
+	return TargetHandler->OnControllerAnalog(KeyName, PlatformUserId, InputDeviceId, AnalogValue);
+}
+
+bool FBridgeMessageHandler::OnControllerButtonPressed(FGamepadKeyNames::Type KeyName, FPlatformUserId PlatformUserId, FInputDeviceId InputDeviceId, bool IsRepeat)
+{
+	return TargetHandler->OnControllerButtonPressed(KeyName, PlatformUserId, InputDeviceId, IsRepeat);
+}
+
+bool FBridgeMessageHandler::OnControllerButtonReleased(FGamepadKeyNames::Type KeyName, FPlatformUserId PlatformUserId, FInputDeviceId InputDeviceId, bool IsRepeat)
+{
+	return TargetHandler->OnControllerButtonReleased(KeyName, PlatformUserId, InputDeviceId, IsRepeat);
+}
 
 void FBridgeMessageHandler::OnBeginGesture()
 {
diff --git a/Engine/Plugins/Bridge/Source/Bridge/Private/UI/FBridgeMessageHandler.h b/Engine/Plugins/Bridge/Source/Bridge/Private/UI/FBridgeMessageHandler.h
--- a/Engine/Plugins/Bridge/Source/Bridge/Private/UI/FBridgeMessageHandler.h
+++ b/Engine/Plugins/Bridge/Source/Bridge/Private/UI/FBridgeMessageHandler.h
@@ -56,12 +56,21 @@ public:
 
 	virtual bool OnCursorSet() override;
 
+	UE_DEPRECATED(5.1, "This version of OnControllerAnalog has been deprecated, please use the one that takes an FPlatformUser and FInputDeviceId instead.")
 	virtual bool OnControllerAnalog(FGamepadKeyNames::Type KeyName, int32 ControllerId, float AnalogValue) override;
 
+	UE_DEPRECATED(5.1, "This version of OnControllerButtonPressed has been deprecated, please use the one that takes an FPlatformUser and FInputDeviceId instead.")
 	virtual bool OnControllerButtonPressed(FGamepadKeyNames::Type KeyName, int32 ControllerId, bool IsRepeat) override;
 
+	UE_DEPRECATED(5.1, "This version of OnControllerButtonReleased has been deprecated, please use the one that takes an FPlatformUser and FInputDeviceId instead.")
 	virtual bool OnControllerButtonReleased(FGamepadKeyNames::Type KeyName, int32 ControllerId, bool IsRepeat) override;
 
+	virtual bool OnControllerAnalog(FGamepadKeyNames::Type KeyName, FPlatformUserId PlatformUserId, FInputDeviceId InputDeviceId, float AnalogValue) override;
+
+	virtual bool OnControllerButtonPressed(FGamepadKeyNames::Type KeyName, FPlatformUserId PlatformUserId, FInputDeviceId InputDeviceId, bool IsRepeat) override;
+
+	virtual bool OnControllerButtonReleased(FGamepadKeyNames::Type KeyName, FPlatformUserId PlatformUserId, FInputDeviceId InputDeviceId, bool IsRepeat) override;
+	
 	virtual void OnBeginGesture() override;
 
 	virtual bool OnTouchGesture(EGestureEvent GestureType, const FVector2D& Delta, float WheelDelta, bool bIsDirectionInvertedFromDevice) override;
