Remove call to IsValid on SharedPtr and addd more error checking

diff --git a/Engine/Plugins/Bridge/Source/Bridge/Private/UI/BridgeUIManager.cpp b/Engine/Plugins/Bridge/Source/Bridge/Private/UI/BridgeUIManager.cpp
--- a/Engine/Plugins/Bridge/Source/Bridge/Private/UI/BridgeUIManager.cpp
+++ b/Engine/Plugins/Bridge/Source/Bridge/Private/UI/BridgeUIManager.cpp
@@ -168,10 +168,17 @@ void FBridgeUIManagerImpl::CreateWindow()
 {
 #if PLATFORM_MAC
 	// Check if WebBrowserWidget plugin is enabled
-	TSharedPtr<IPlugin> WebBrowserPlugin = IPluginManager::Get().FindPlugin("WebBrowserWidget");
-	if (WebBrowserPlugin.IsValid() && !WebBrowserPlugin->IsEnabled())
+	if (TSharedPtr<IPlugin> WebBrowserPlugin = IPluginManager::Get().FindPlugin("WebBrowserWidget"))
 	{
-		FMessageDialog::Open(EAppMsgType::Ok, LOCTEXT("WebBrowserWidgetPluginNotEnabled", "Web Browser plugin is not enabled. Please enable it in the plugin manager to use Bridge."));
+		if (!WebBrowserPlugin->IsEnabled())
+		{
+			FMessageDialog::Open(EAppMsgType::Ok, LOCTEXT("WebBrowserWidgetPluginNotEnabled", "Web Browser plugin is not enabled. Please enable it in the plugin manager to use Bridge."));
+			return;
+		}
+	}
+	else
+	{
+		FMessageDialog::Open(EAppMsgType::Ok, LOCTEXT("WebBrowserWidgetPluginNotFound", "Web Browser plugin is not found."));
 		return;
 	}
 #endif
@@ -227,11 +234,19 @@ TSharedRef<SDockTab> FBridgeUIManagerImpl::CreateBridgeTab(const FSpawnTabArgs&
 {
 #if PLATFORM_MAC
 	// Check if WebBrowserWidget plugin is enabled
-	TSharedPtr<IPlugin> WebBrowserPlugin = IPluginManager::Get().FindPlugin("WebBrowserWidget");
-	if (WebBrowserPlugin.IsValid() && !WebBrowserPlugin->IsEnabled())
+	if (TSharedPtr<IPlugin> WebBrowserPlugin = IPluginManager::Get().FindPlugin("WebBrowserWidget"))
 	{
-		FMessageDialog::Open(EAppMsgType::Ok, LOCTEXT("WebBrowserWidgetPluginNotEnabled", "Web Browser plugin is not enabled. Please enable it in the plugin manager to use Bridge."));
+		if (!WebBrowserPlugin->IsEnabled())
+		{
+			FMessageDialog::Open(EAppMsgType::Ok, LOCTEXT("WebBrowserWidgetPluginNotEnabled", "Web Browser plugin is not enabled. Please enable it in the plugin manager to use Bridge."));
 
+			return SAssignNew(LocalBrowserDock, SDockTab)
+				.TabRole(ETabRole::NomadTab);
+		}
+	}
+	else
+	{
+		FMessageDialog::Open(EAppMsgType::Ok, LOCTEXT("WebBrowserWidgetPluginNotFound", "Web Browser plugin is not found."));
 		return SAssignNew(LocalBrowserDock, SDockTab)
 			.TabRole(ETabRole::NomadTab);
 	}
