Quixel Bridge - fix crash on export when Virtual Textures are enabled

diff --git a/Engine/Plugins/Bridge/Source/MegascansPlugin/Private/Utilities/MiscUtils.cpp b/Engine/Plugins/Bridge/Source/MegascansPlugin/Private/Utilities/MiscUtils.cpp
--- a/Engine/Plugins/Bridge/Source/MegascansPlugin/Private/Utilities/MiscUtils.cpp
+++ b/Engine/Plugins/Bridge/Source/MegascansPlugin/Private/Utilities/MiscUtils.cpp
@@ -53,18 +53,15 @@
 
 #include "JsonObjectConverter.h"
 
-
 TSharedPtr<FJsonObject> DeserializeJson(const FString& JsonStringData)
 {
 	TSharedPtr<FJsonObject> JsonDataObject;
-	bool bIsDeseriliazeSuccessful = false;	
+	bool bIsDeseriliazeSuccessful = false;
 	TSharedRef<TJsonReader<>> JsonReader = TJsonReaderFactory<>::Create(JsonStringData);
-	bIsDeseriliazeSuccessful = FJsonSerializer::Deserialize(JsonReader, JsonDataObject);	
+	bIsDeseriliazeSuccessful = FJsonSerializer::Deserialize(JsonReader, JsonDataObject);
 	return JsonDataObject;
 }
 
-
-
 FString GetPluginPath()
 {
 	FString PluginName = TEXT("Bridge");
@@ -80,20 +77,16 @@ FString GetSourceMSPresetsPath()
 	return MaterialPresetsPath;
 }
 
-
-
 FString GetMSPresetsName()
 {
 	return TEXT("MSPresets");
 }
 
-
-
-bool CopyMaterialPreset(const FString & MaterialName)
-{	
+bool CopyMaterialPreset(const FString& MaterialName)
+{
 	FString MaterialSourceFolderPath = FPaths::Combine(GetSourceMSPresetsPath(), MaterialName);
-	MaterialSourceFolderPath = FPaths::ConvertRelativePathToFull(MaterialSourceFolderPath);	
-	IPlatformFile& PlatformFile = FPlatformFileManager::Get().GetPlatformFile();	
+	MaterialSourceFolderPath = FPaths::ConvertRelativePathToFull(MaterialSourceFolderPath); 
+	IPlatformFile& PlatformFile = FPlatformFileManager::Get().GetPlatformFile(); 
 	if (!PlatformFile.DirectoryExists(*MaterialSourceFolderPath))
 	{
 		return false;
@@ -105,22 +98,17 @@ bool CopyMaterialPreset(const FString & MaterialName)
 	
 	if (!PlatformFile.DirectoryExists(*MaterialDestinationPath))
 	{
-		
 		if (!PlatformFile.CreateDirectoryTree(*MaterialDestinationPath))
 		{
-			
 			return false;
 		}
 
 		if (!PlatformFile.CopyDirectoryTree(*MaterialDestinationPath, *MaterialSourceFolderPath, true))
 		{
-			
 			return false;
 		}
-
 	}
 
-	
 	FAssetRegistryModule& AssetRegistryModule = FModuleManager::GetModuleChecked<FAssetRegistryModule>("AssetRegistry");
 	TArray<FString> MaterialBasePath;
 	FString BasePath = FPaths::Combine(FPaths::ProjectContentDir(), GetMSPresetsName());
@@ -132,31 +120,25 @@ bool CopyMaterialPreset(const FString & MaterialName)
 	}
 	AssetRegistryModule.Get().ScanPathsSynchronous(MaterialBasePath, true);
 	return true;
-	
 }
 
-
-
-
 bool CopyPresetTextures()
 {
 	FString TexturesDestinationPath = FPaths::ProjectContentDir();
-	
+
 	TexturesDestinationPath = FPaths::Combine(TexturesDestinationPath, GetMSPresetsName(), TEXT("MSTextures"));
 
-	FString TexturesSourceFolderPath = FPaths::Combine(GetSourceMSPresetsPath(), TEXT("MSTextures"));	
+	FString TexturesSourceFolderPath = FPaths::Combine(GetSourceMSPresetsPath(), TEXT("MSTextures"));
 	IPlatformFile& PlatformFile = FPlatformFileManager::Get().GetPlatformFile();
 	if (!PlatformFile.DirectoryExists(*TexturesDestinationPath))
 	{
 		if (!PlatformFile.CreateDirectoryTree(*TexturesDestinationPath))
 		{
-			
 			return false;
 		}
 
 		if (!PlatformFile.CopyDirectoryTree(*TexturesDestinationPath, *TexturesSourceFolderPath, true))
 		{
-			
 			return false;
 		}
 	}
@@ -164,8 +146,6 @@ bool CopyPresetTextures()
 	return true;
 }
 
-
-
 void CopyMSPresets()
 {
 	FString MaterialsDestination = FPaths::ProjectContentDir();
@@ -173,12 +153,20 @@ void CopyMSPresets()
 	FString MaterialsSourceFolderPath = GetSourceMSPresetsPath();
 	IPlatformFile& PlatformFile = FPlatformFileManager::Get().GetPlatformFile();
 
-	
+	PlatformFile.CreateDirectoryTree(*MaterialsDestination);
 
+	TArray<FString> SourceFiles;
+	PlatformFile.FindFilesRecursively(SourceFiles, *MaterialsSourceFolderPath, NULL);
 
-	
-	PlatformFile.CreateDirectoryTree(*MaterialsDestination);
-	PlatformFile.CopyDirectoryTree(*MaterialsDestination, *MaterialsSourceFolderPath, false);
+	for (FString SourceFile : SourceFiles)
+	{
+		FString DestinationFile = FPaths::Combine(MaterialsDestination, SourceFile.Replace(*MaterialsSourceFolderPath, TEXT("")));
+
+		FString BaseDestinationDirectory = FPaths::GetPath(DestinationFile);
+		PlatformFile.CreateDirectoryTree(*BaseDestinationDirectory);
+
+		PlatformFile.CopyFile(*DestinationFile, *SourceFile);
+	}
 
 	FAssetRegistryModule& AssetRegistryModule = FModuleManager::GetModuleChecked<FAssetRegistryModule>("AssetRegistry");
 	IAssetRegistry& AssetRegistry = AssetRegistryModule.Get();
@@ -186,161 +174,136 @@ void CopyMSPresets()
 	TArray<FString> SyncPaths;
 	SyncPaths.Add(TEXT("/Game/MSPresets"));
 	AssetRegistryModule.Get().ScanPathsSynchronous(SyncPaths, true);
+}
 
-	
+void AssetUtils::FocusOnSelected(const FString& Path)
+{
+	TArray<FString> Folders;
+	Folders.Add(Path);
+	FContentBrowserModule& ContentBrowserModule = FModuleManager::LoadModuleChecked<FContentBrowserModule>("ContentBrowser");
+	IContentBrowserSingleton& ContentBrowserSingleton = ContentBrowserModule.Get();
+	ContentBrowserSingleton.SyncBrowserToFolders(Folders);
+}
 
+void AssetUtils::SavePackage(UObject* SourceObject)
+{
+	/*TArray<UObject*> InputObjects;
+	InputObjects.Add(SourceObject);
+	UPackageTools::SavePackagesForObjects(InputObjects);*/
 
+	TArray<UPackage*> PackagesToSave;
+	PackagesToSave.Add(SourceObject->GetPackage());
+	UEditorLoadingAndSavingUtils::SavePackages(PackagesToSave, true);
 }
 
+void AssetUtils::DeleteDirectory(FString TargetDirectory)
+{
+	TArray<FString> PreviewAssets = UEditorAssetLibrary::ListAssets(TargetDirectory);
 
-	void AssetUtils::FocusOnSelected(const FString& Path)
+	for (FString AssetPath : PreviewAssets)
 	{
-		TArray<FString> Folders;
-		Folders.Add(Path);
-		FContentBrowserModule& ContentBrowserModule = FModuleManager::LoadModuleChecked<FContentBrowserModule>("ContentBrowser");
-		IContentBrowserSingleton& ContentBrowserSingleton = ContentBrowserModule.Get();
-		ContentBrowserSingleton.SyncBrowserToFolders(Folders);
+		DeleteAsset(AssetPath);
 	}
 
+	UEditorAssetLibrary::DeleteDirectory(TargetDirectory);
+}
 
+bool AssetUtils::DeleteAsset(const FString& AssetPath)
+{
+	IAssetRegistry& AssetRegistry = FModuleManager::GetModuleChecked<FAssetRegistryModule>("AssetRegistry").Get();
+	FAssetData AssetData = AssetRegistry.GetAssetByObjectPath(FName(*AssetPath));
 
-	void AssetUtils::SavePackage(UObject* SourceObject)
+	if (AssetData.IsAssetLoaded())
 	{
-		/*TArray<UObject*> InputObjects;
-		InputObjects.Add(SourceObject);
-		UPackageTools::SavePackagesForObjects(InputObjects);*/
-
-		TArray<UPackage*> PackagesToSave;
-		PackagesToSave.Add(SourceObject->GetPackage());
-		UEditorLoadingAndSavingUtils::SavePackages(PackagesToSave, true);
-
-	}
-
-	void AssetUtils::DeleteDirectory(FString TargetDirectory)
-	{		
-		
-		TArray<FString> PreviewAssets = UEditorAssetLibrary::ListAssets(TargetDirectory);
-
-		for (FString AssetPath : PreviewAssets)
+		if (UEditorAssetLibrary::DeleteLoadedAsset(AssetData.GetAsset()))
 		{
-			DeleteAsset(AssetPath);
+			return true;
 		}
-
-		UEditorAssetLibrary::DeleteDirectory(TargetDirectory);
-
-
 	}
-
-	bool AssetUtils::DeleteAsset(const FString& AssetPath)
+	else
 	{
-		IAssetRegistry& AssetRegistry = FModuleManager::GetModuleChecked<FAssetRegistryModule>("AssetRegistry").Get();
-		FAssetData AssetData = AssetRegistry.GetAssetByObjectPath(FName(*AssetPath));
-
-		if (AssetData.IsAssetLoaded())
+		if (UEditorAssetLibrary::DeleteAsset(AssetPath))
 		{
-			if (UEditorAssetLibrary::DeleteLoadedAsset(AssetData.GetAsset()))
-			{
-				return true;
-			}
+			return true;
 		}
-		else
-		{
-			if (UEditorAssetLibrary::DeleteAsset(AssetPath))
-			{
-				return true;
-			}
-		}
-
-		return false;
 	}
 
-	
-		
-
-	
-
-	FUAssetMeta AssetUtils::GetAssetMetaData(const FString& JsonPath)
-	{
-		
-		FString UassetMetaString;
-		FFileHelper::LoadFileToString(UassetMetaString, *JsonPath);
-		FUAssetMeta AssetMetaData;
-		FJsonObjectConverter::JsonObjectStringToUStruct(UassetMetaString, &AssetMetaData);
-		
-		return AssetMetaData;
-	}
+	return false;
+}
 
-	TArray<UMaterialInstanceConstant*> AssetUtils::GetSelectedAssets(const FTopLevelAssetPath& AssetClass)
-	{
-		check(!AssetClass.IsNull());
+FUAssetMeta AssetUtils::GetAssetMetaData(const FString& JsonPath)
+{
+	FString UassetMetaString;
+	FFileHelper::LoadFileToString(UassetMetaString, *JsonPath);
+	FUAssetMeta AssetMetaData;
+	FJsonObjectConverter::JsonObjectStringToUStruct(UassetMetaString, &AssetMetaData);
 
-		TArray<UMaterialInstanceConstant*> ObjectArray;
+	return AssetMetaData;
+}
 
-		TArray<FAssetData> AssetDatas;
-		FContentBrowserModule& ContentBrowserModule = FModuleManager::LoadModuleChecked<FContentBrowserModule>("ContentBrowser");
-		IContentBrowserSingleton& ContentBrowserSingleton = ContentBrowserModule.Get();
-		ContentBrowserSingleton.GetSelectedAssets(AssetDatas);
+TArray<UMaterialInstanceConstant*> AssetUtils::GetSelectedAssets(const FTopLevelAssetPath& AssetClass)
+{
+	check(!AssetClass.IsNull());
 
-		for (FAssetData SelectedAsset : AssetDatas)
-		{
+	TArray<UMaterialInstanceConstant*> ObjectArray;
 
-			if (SelectedAsset.AssetClassPath == AssetClass)
-			{
-				ObjectArray.Add(CastChecked<UMaterialInstanceConstant>(UEditorAssetLibrary::LoadAsset(SelectedAsset.ObjectPath.ToString())));
-			}
-		}
-		return ObjectArray;
-	}
+	TArray<FAssetData> AssetDatas;
+	FContentBrowserModule& ContentBrowserModule = FModuleManager::LoadModuleChecked<FContentBrowserModule>("ContentBrowser");
+	IContentBrowserSingleton& ContentBrowserSingleton = ContentBrowserModule.Get();
+	ContentBrowserSingleton.GetSelectedAssets(AssetDatas);
 
-	void AssetUtils::AddFoliageTypesToLevel(TArray<FString> FoliageTypePaths)
+	for (FAssetData SelectedAsset : AssetDatas)
 	{
-		IAssetRegistry& AssetRegistry = FModuleManager::GetModuleChecked<FAssetRegistryModule>("AssetRegistry").Get();
-		
-
-		for (FString FoliageTypePath : FoliageTypePaths)
+		if (SelectedAsset.AssetClassPath == AssetClass)
 		{
-			FAssetData AssetData = AssetRegistry.GetAssetByObjectPath(FName(*FoliageTypePath));
-			if (!AssetData.IsValid()) return;
-
-			auto* CurrentWorld = GEditor->GetEditorWorldContext().World();
-			AInstancedFoliageActor* IFA = AInstancedFoliageActor::GetInstancedFoliageActorForCurrentLevel(CurrentWorld, true);
-			IFA->AddFoliageType(Cast<UFoliageType_InstancedStaticMesh>(AssetData.GetAsset()));
+			ObjectArray.Add(CastChecked<UMaterialInstanceConstant>(UEditorAssetLibrary::LoadAsset(SelectedAsset.ObjectPath.ToString())));
 		}
 	}
 
-	void AssetUtils::ManageImportSettings(FUAssetMeta AssetMetaData)
-	{
-		const UMegascansSettings* MegascansSettings = GetDefault<UMegascansSettings>();
+	return ObjectArray;
+}
 
-		if ((AssetMetaData.assetType == TEXT("surface") || AssetMetaData.assetType == TEXT("atlas")) && MegascansSettings->bApplyToSelection)
-		{			
-			
-			FMaterialUtils::ApplyMaterialToSelection(AssetMetaData.materialInstances[0].instancePath);
-			
-		}
+void AssetUtils::AddFoliageTypesToLevel(TArray<FString> FoliageTypePaths)
+{
+	IAssetRegistry& AssetRegistry = FModuleManager::GetModuleChecked<FAssetRegistryModule>("AssetRegistry").Get();
 
-		if (AssetMetaData.assetType == TEXT("3dplant") && MegascansSettings->bCreateFoliage)
-		{
-			AssetUtils::AddFoliageTypesToLevel(AssetMetaData.foliageAssetPaths);
-		}
+	for (FString FoliageTypePath : FoliageTypePaths)
+	{
+		FAssetData AssetData = AssetRegistry.GetAssetByObjectPath(FName(*FoliageTypePath));
+		if (!AssetData.IsValid()) return;
 
+		auto* CurrentWorld = GEditor->GetEditorWorldContext().World();
+		AInstancedFoliageActor* IFA = AInstancedFoliageActor::GetInstancedFoliageActorForCurrentLevel(CurrentWorld, true);
+		IFA->AddFoliageType(Cast<UFoliageType_InstancedStaticMesh>(AssetData.GetAsset()));
 	}
+}
 
+void AssetUtils::ManageImportSettings(FUAssetMeta AssetMetaData)
+{
+	const UMegascansSettings* MegascansSettings = GetDefault<UMegascansSettings>();
 
+	if ((AssetMetaData.assetType == TEXT("surface") || AssetMetaData.assetType == TEXT("atlas")) && MegascansSettings->bApplyToSelection)
+	{
+		FMaterialUtils::ApplyMaterialToSelection(AssetMetaData.materialInstances[0].instancePath);
+	}
 
-
+	if (AssetMetaData.assetType == TEXT("3dplant") && MegascansSettings->bCreateFoliage)
+	{
+		AssetUtils::AddFoliageTypesToLevel(AssetMetaData.foliageAssetPaths);
+	}
+}
 
 // Get import type based on incoming json
 EAssetImportType JsonUtils::GetImportType(TSharedPtr<FJsonObject> ImportJsonObject)
 {
 	FString ExportType;
 	if (ImportJsonObject->TryGetStringField(TEXT("exportType"), ExportType))
-	{		
+	{
 		return (ExportType == TEXT("megascans_uasset")) ? EAssetImportType::MEGASCANS_UASSET : (ExportType == TEXT("megascans_source")) ? EAssetImportType::MEGASCANS_SOURCE : (ExportType == TEXT("dhi")) ? EAssetImportType::DHI_CHARACTER : (ExportType == TEXT("template")) ? EAssetImportType::TEMPLATE : EAssetImportType::NONE;
 
 	}
 	else
-	{	
+	{
 		return EAssetImportType::NONE;
 	}
 }
@@ -355,7 +318,7 @@ TSharedPtr<FUAssetData> JsonUtils::ParseUassetJson(TSharedPtr<FJsonObject> Impor
 	ImportData->ImportJsonPath = ImportJsonObject->GetStringField(TEXT("importJson"));
 	ImportData->ImportType = ImportJsonObject->GetStringField(TEXT("exportType"));
 	ImportData->AssetId = ImportJsonObject->GetStringField(TEXT("assetId"));
-	ImportData->ProgressiveStage = ImportJsonObject->GetIntegerField (TEXT("progressiveStage"));
+	ImportData->ProgressiveStage = ImportJsonObject->GetIntegerField(TEXT("progressiveStage"));
 
 	TArray<TSharedPtr<FJsonValue>> FilePathsArray = ImportJsonObject->GetArrayField(TEXT("assetPaths"));
 
@@ -367,7 +330,6 @@ TSharedPtr<FUAssetData> JsonUtils::ParseUassetJson(TSharedPtr<FJsonObject> Impor
 	ImportData->FilePaths = AssetPaths;
 
 	return ImportData;
-	
 }
 
 void CopyUassetFiles(TArray<FString> FilesToCopy, FString DestinationDirectory)
@@ -377,21 +339,15 @@ void CopyUassetFiles(TArray<FString> FilesToCopy, FString DestinationDirectory)
 	FAssetRegistryModule& AssetRegistryModule = FModuleManager::GetModuleChecked<FAssetRegistryModule>("AssetRegistry");
 	IAssetRegistry& AssetRegistry = AssetRegistryModule.Get();
 
-	
-
 	for (FString FileToCopy : FilesToCopy)
 	{
 		FString DestinationFile = FPaths::Combine(DestinationDirectory, FPaths::GetCleanFilename(FileToCopy));
 		PlatformFile.CopyFile(*DestinationFile, *FileToCopy);
 	}
 
-
-	
 	TArray<FString> SyncPaths;
 	SyncPaths.Add(TEXT("/Game/Megascans"));
 	AssetRegistryModule.Get().ScanPathsSynchronous(SyncPaths, true);
-	
-
 }
 
 void AssetUtils::SyncFolder(const FString& TargetFolder)
@@ -401,9 +357,7 @@ void AssetUtils::SyncFolder(const FString& TargetFolder)
 	SyncPaths.Add(TargetFolder);
 	AssetRegistry.ScanPathsSynchronous(SyncPaths, true);
 
-	//Testing syncing 
-	
-
+	// Testing syncing
 }
 
 void AssetUtils::RegisterAsset(const FString& PackagePath)
@@ -420,26 +374,25 @@ void AssetUtils::ConvertToVT(FUAssetMeta AssetMetaData)
 
 		FMaterialBlend::Get()->ConvertToVirtualTextures(AssetMetaData);
 
-		//Call Master material override maybe to avoid all other steps
-		//Create a new instance of the VT enabled master material based on which original material was used
-		//Plug all the textures to this new material instance
-		//Apply the newly generated material instance to the mesh incase of 3d assets
-
+		// Call Master material override maybe to avoid all other steps
+		// Create a new instance of the VT enabled master material based on which original material was used
+		// Plug all the textures to this new material instance
+		// Apply the newly generated material instance to the mesh incase of 3d assets
 	}
 }
 
-
 bool AssetUtils::IsVTEnabled()
 {
+	static const auto CVarVirtualTexturesEnabled = IConsoleManager::Get().FindTConsoleVariableDataInt(TEXT("r.VirtualTextures"));
+	check(CVarVirtualTexturesEnabled);
+	static const auto CVarVirtualTexturesImportEnabled = IConsoleManager::Get().FindTConsoleVariableDataInt(TEXT("r.VT.EnableAutoImport"));
+	check(CVarVirtualTexturesEnabled);
 
-	static const auto CVarVirtualTexturesEnabled = IConsoleManager::Get().FindTConsoleVariableDataInt(TEXT("r.VirtualTextures")); check(CVarVirtualTexturesEnabled);
-	static const auto CVarVirtualTexturesImportEnabled = IConsoleManager::Get().FindTConsoleVariableDataInt(TEXT("r.VT.EnableAutoImport")); check(CVarVirtualTexturesEnabled);
-
-	
 	const bool bVirtualTextureEnabled = CVarVirtualTexturesEnabled->GetValueOnAnyThread() != 0;
 	const bool bVirtualTextureImportEnabled = CVarVirtualTexturesImportEnabled->GetValueOnAnyThread() != 0;
 
-	if (bVirtualTextureEnabled && bVirtualTextureImportEnabled) {
+	if (bVirtualTextureEnabled && bVirtualTextureImportEnabled)
+	{
 		return true;
 	}
 
@@ -496,17 +449,15 @@ void UpdateMHVersionInfo(TMap<FString, TArray<FString>> AssetsStatus, TMap<FStri
 		{
 			AssetInfo.version = FString::SanitizeFloat(SourceAssetsVersionInfo[AssetInfo.path]);
 		}
-		
 	}
 
-
 	for (FString AssetToAdd : AssetsStatus["Add"])
 	{
 		FAssetInfo NewAssetInfo;
 		NewAssetInfo.path = AssetToAdd;
 		if (SourceAssetsVersionInfo.Contains(AssetToAdd))
-		{			
-			NewAssetInfo.version = FString::SanitizeFloat(SourceAssetsVersionInfo[AssetToAdd]);			
+		{
+			NewAssetInfo.version = FString::SanitizeFloat(SourceAssetsVersionInfo[AssetToAdd]);
 		}
 		else
 		{
@@ -515,12 +466,8 @@ void UpdateMHVersionInfo(TMap<FString, TArray<FString>> AssetsStatus, TMap<FStri
 		VersionData.assets.Add(NewAssetInfo);
 	}
 
-
 	FString OutputData;
 	FJsonObjectConverter::UStructToJsonObjectString(VersionData, OutputData);
 
-
 	FFileHelper::SaveStringToFile(OutputData, *ProjectAssetsVersionPath);
-
-
 }
