FNames containing asset paths are deprecated. FSoftObjectPath or FTopLevelAssetPath should be used instead.
Fixups for Megascans.

diff --git a/Engine/Plugins/Bridge/Source/MegascansPlugin/Private/AssetImporters/DHIImport.cpp b/Engine/Plugins/Bridge/Source/MegascansPlugin/Private/AssetImporters/DHIImport.cpp
--- a/Engine/Plugins/Bridge/Source/MegascansPlugin/Private/AssetImporters/DHIImport.cpp
+++ b/Engine/Plugins/Bridge/Source/MegascansPlugin/Private/AssetImporters/DHIImport.cpp
@@ -618,7 +618,7 @@ void FImportDHI::ImportAsset(TSharedPtr<FJsonObject> AssetImportJson)
 
 				FString AssetPackagePath = FPaths::Combine(TEXT("/Game"), StrippedFilePath);
 
-				FAssetData GameAssetData = AssetRegistryModule.Get().GetAssetByObjectPath(FName(*AssetPackagePath));
+				FAssetData GameAssetData = AssetRegistryModule.Get().GetAssetByObjectPath(FSoftObjectPath(AssetPackagePath));
 				
 				if (GameAssetData.IsAssetLoaded())
 				{					
@@ -729,7 +729,7 @@ void FImportDHI::ImportAsset(TSharedPtr<FJsonObject> AssetImportJson)
 
     AssetRegistryModule.Get().ScanPathsSynchronous(AssetsBasePath, true);
 
-    FAssetData CharacterAssetData = AssetRegistryModule.Get().GetAssetByObjectPath(FName(*BPPath));
+    FAssetData CharacterAssetData = AssetRegistryModule.Get().GetAssetByObjectPath(FSoftObjectPath(BPPath));
 
     // As part of our discussion to fix this, here was the explanation about why the asset has to be loaded:
     // "In UE4 we came across this issue where even if we did call the syncfolder on content browser, it would still
@@ -1030,7 +1030,7 @@ void FImportDHI::CopyCommonFiles(TArray<FString> AssetsList, const FString& Comm
 				// Create a path like /Game/Metahumans/MetaHumans/Common/Common/MetaHuman_ControlRig.MetaHuman_ControlRig
 				FString GameAssetPath = TEXT("/Game/") + AssetRelativePath + TEXT(".") + AssetName;
 
-				FAssetData GameAssetData = AssetRegistryModule.Get().GetAssetByObjectPath(FName(*GameAssetPath));
+				FAssetData GameAssetData = AssetRegistryModule.Get().GetAssetByObjectPath(FSoftObjectPath(GameAssetPath));
 
 				if (GameAssetData.IsAssetLoaded())
 				{
@@ -1151,7 +1151,7 @@ TArray<FString> FImportDHI::CharactersOpenEditors()
 
 	for (FString CharacterInProject : CharactersInProject)
 	{
-		FAssetData CharacterAssetData = AssetRegistryModule.Get().GetAssetByObjectPath(FName(*CharacterInProject));
+		FAssetData CharacterAssetData = AssetRegistryModule.Get().GetAssetByObjectPath(FSoftObjectPath(CharacterInProject));
 
 		if (CharacterAssetData.IsAssetLoaded())
 		{
@@ -1172,7 +1172,7 @@ void FImportDHI::CloseAssetEditors(TArray<FString> CharactersOpenInEditors)
 
 	for (FString CharacterAsset : CharactersOpenInEditors)
 	{
-		FAssetData CharacterAssetData = AssetRegistryModule.Get().GetAssetByObjectPath(FName(*CharacterAsset));
+		FAssetData CharacterAssetData = AssetRegistryModule.Get().GetAssetByObjectPath(FSoftObjectPath(CharacterAsset));
 		
 		UObject* CharacterObject = CharacterAssetData.GetAsset();
 
@@ -1188,7 +1188,7 @@ void FImportDHI::RecompileAllCharacters()
 
 	for (FString CharacterInProject : CharactersInProject)
 	{
-		FAssetData GameAssetData = AssetRegistryModule.Get().GetAssetByObjectPath(FName(*CharacterInProject));
+		FAssetData GameAssetData = AssetRegistryModule.Get().GetAssetByObjectPath(FSoftObjectPath(CharacterInProject));
 		UObject* ItemObject = GameAssetData.GetAsset();
 
 		if (UBlueprint* BPObject = Cast<UBlueprint>(ItemObject))
diff --git a/Engine/Plugins/Bridge/Source/MegascansPlugin/Private/AssetImporters/ProgressiveImport3D.cpp b/Engine/Plugins/Bridge/Source/MegascansPlugin/Private/AssetImporters/ProgressiveImport3D.cpp
--- a/Engine/Plugins/Bridge/Source/MegascansPlugin/Private/AssetImporters/ProgressiveImport3D.cpp
+++ b/Engine/Plugins/Bridge/Source/MegascansPlugin/Private/AssetImporters/ProgressiveImport3D.cpp
@@ -113,7 +113,7 @@ void FImportProgressive3D::ImportAsset(TSharedPtr<FJsonObject> AssetImportJson,
 	if (bIsNormal)
 	{
 		FString MeshPath = AssetMetaData.meshList[0].path;
-		FAssetData AssetMeshData = AssetRegistry.GetAssetByObjectPath(FName(*MeshPath));
+		FAssetData AssetMeshData = AssetRegistry.GetAssetByObjectPath(FSoftObjectPath(MeshPath));
 		FSoftObjectPath ItemToStream = AssetMeshData.ToSoftObjectPath();
 		if (!AssetMeshData.IsValid()) return;
 
@@ -142,10 +142,10 @@ void FImportProgressive3D::ImportAsset(TSharedPtr<FJsonObject> AssetImportJson,
 
 		FString MeshPath = AssetMetaData.meshList[0].path;
 		PreviewDetails[ImportData->AssetId]->PreviewMeshPath = MeshPath;
-		FAssetData PreviewMeshData = AssetRegistry.GetAssetByObjectPath(FName(*MeshPath));
+		FAssetData PreviewMeshData = AssetRegistry.GetAssetByObjectPath(FSoftObjectPath(MeshPath));
 		FString MInstancePath = AssetMetaData.materialInstances[0].instancePath ;
 		
-		FAssetData MInstanceData = AssetRegistry.GetAssetByObjectPath(FName(*MInstancePath));
+		FAssetData MInstanceData = AssetRegistry.GetAssetByObjectPath(FSoftObjectPath(MInstancePath));
 		FSoftObjectPath ItemToStream = PreviewMeshData.ToSoftObjectPath();
 		
 		if (!MInstanceData.IsValid()) return;		
@@ -157,7 +157,7 @@ void FImportProgressive3D::ImportAsset(TSharedPtr<FJsonObject> AssetImportJson,
 
 		Streamable.RequestAsyncLoad(ItemToStream, FStreamableDelegate::CreateRaw(this, &FImportProgressive3D::HandlePreviewInstanceLoad, MInstanceData, ImportData->AssetId));
 
-		FString AssetPath = PreviewMeshData.ObjectPath.ToString();
+		FString AssetPath = PreviewMeshData.GetObjectPathString();
 		FAssetData DraggedAssetData = FAssetData(LoadObject<UStaticMesh>(nullptr, *AssetPath));
 
 		FBridgeDragDropHelper::Instance->OnAddProgressiveStageDataDelegate.ExecuteIfBound(DraggedAssetData, ImportData->AssetId, AssetMetaData.assetType, nullptr);
@@ -176,7 +176,7 @@ void FImportProgressive3D::ImportAsset(TSharedPtr<FJsonObject> AssetImportJson,
 		}
 
 		
-		FAssetData AlbedoData = AssetRegistry.GetAssetByObjectPath(FName(*AlbedoPath));
+		FAssetData AlbedoData = AssetRegistry.GetAssetByObjectPath(FSoftObjectPath(AlbedoPath));
 		FSoftObjectPath ItemToStream = AlbedoData.ToSoftObjectPath();
 
 		if (!AlbedoData.IsValid()) return;
@@ -199,7 +199,7 @@ void FImportProgressive3D::ImportAsset(TSharedPtr<FJsonObject> AssetImportJson,
 		}
 		
 		
-		FAssetData NormalData = AssetRegistry.GetAssetByObjectPath(FName(*NormalPath));
+		FAssetData NormalData = AssetRegistry.GetAssetByObjectPath(FSoftObjectPath(NormalPath));
 		FSoftObjectPath ItemToStream = NormalData.ToSoftObjectPath();
 
 		if (!NormalData.IsValid()) return;
@@ -213,7 +213,7 @@ void FImportProgressive3D::ImportAsset(TSharedPtr<FJsonObject> AssetImportJson,
 		bool bWaitNaniteConversion = false;
 		FString MeshPath = AssetMetaData.meshList[0].path;
 
-		FAssetData HighMeshData = AssetRegistry.GetAssetByObjectPath(FName(*MeshPath));
+		FAssetData HighMeshData = AssetRegistry.GetAssetByObjectPath(FSoftObjectPath(MeshPath));
 		if (ImportData->AssetType == TEXT("3d") && ImportData->AssetTier==0 && AssetMetaData.assetSubType == TEXT("singleMesh"))
 		{
 			bWaitNaniteConversion = true;
@@ -223,7 +223,7 @@ void FImportProgressive3D::ImportAsset(TSharedPtr<FJsonObject> AssetImportJson,
 
 		if (!HighMeshData.IsValid()) return;
 
-		FAssetData PreviewMeshData = AssetRegistry.GetAssetByObjectPath(FName(*MeshPath));
+		FAssetData PreviewMeshData = AssetRegistry.GetAssetByObjectPath(FSoftObjectPath(MeshPath));
 
 		Streamable.RequestAsyncLoad(ItemToStream, FStreamableDelegate::CreateRaw(this, &FImportProgressive3D::HandleHighAssetLoad, HighMeshData, ImportData->AssetId, AssetMetaData, bWaitNaniteConversion));
 	}
diff --git a/Engine/Plugins/Bridge/Source/MegascansPlugin/Private/AssetImporters/ProgressiveImportSurfaces.cpp b/Engine/Plugins/Bridge/Source/MegascansPlugin/Private/AssetImporters/ProgressiveImportSurfaces.cpp
--- a/Engine/Plugins/Bridge/Source/MegascansPlugin/Private/AssetImporters/ProgressiveImportSurfaces.cpp
+++ b/Engine/Plugins/Bridge/Source/MegascansPlugin/Private/AssetImporters/ProgressiveImportSurfaces.cpp
@@ -73,7 +73,7 @@ void FImportProgressiveSurfaces::ImportAsset(TSharedPtr<FJsonObject> AssetImport
 		if (bIsNormal)
 		{
 			FString InstancePath = AssetMetaData.materialInstances[0].instancePath;
-			FAssetData AssetData = AssetRegistry.GetAssetByObjectPath(FName(*InstancePath));
+			FAssetData AssetData = AssetRegistry.GetAssetByObjectPath(FSoftObjectPath(InstancePath));
 			FBridgeDragDropHelper::Instance->OnAddProgressiveStageDataDelegate.ExecuteIfBound(AssetData, ImportData->AssetId, "decal-normal", nullptr);
 			return;
 		}
@@ -87,13 +87,13 @@ void FImportProgressiveSurfaces::ImportAsset(TSharedPtr<FJsonObject> AssetImport
 		if (ImportData->ProgressiveStage == 1)
 		{
 			FString InstancePath = AssetMetaData.materialInstances[0].instancePath;
-			FAssetData AssetData = AssetRegistry.GetAssetByObjectPath(FName(*InstancePath));
+			FAssetData AssetData = AssetRegistry.GetAssetByObjectPath(FSoftObjectPath(InstancePath));
 			FBridgeDragDropHelper::Instance->OnAddProgressiveStageDataDelegate.ExecuteIfBound(AssetData, ImportData->AssetId, TEXT("decal-stage-1"), nullptr);
 		}
 		if (ImportData->ProgressiveStage == 4)
 		{
 			 FString MInstanceHighPath = AssetMetaData.materialInstances[0].instancePath;
-			 FAssetData MInstanceHighData = AssetRegistry.GetAssetByObjectPath(FName(*MInstanceHighPath));
+			 FAssetData MInstanceHighData = AssetRegistry.GetAssetByObjectPath(FSoftObjectPath(MInstanceHighPath));
 			 FBridgeDragDropHelper::Instance->OnAddProgressiveStageDataDelegate.ExecuteIfBound(MInstanceHighData, ImportData->AssetId, TEXT("decal-stage-4"), nullptr);
 		}
 
@@ -103,7 +103,7 @@ void FImportProgressiveSurfaces::ImportAsset(TSharedPtr<FJsonObject> AssetImport
 	if (bIsNormal)
 	{
 		FString MInstancePath = AssetMetaData.materialInstances[0].instancePath;
-		FAssetData MInstanceData = AssetRegistry.GetAssetByObjectPath(FName(*MInstancePath));
+		FAssetData MInstanceData = AssetRegistry.GetAssetByObjectPath(FSoftObjectPath(MInstancePath));
 
 		if (!MInstanceData.IsValid()) return;
 
@@ -132,7 +132,7 @@ void FImportProgressiveSurfaces::ImportAsset(TSharedPtr<FJsonObject> AssetImport
 	if (ImportData->ProgressiveStage == 1)
 	{
 		FString MInstancePath = AssetMetaData.materialInstances[0].instancePath;
-		FAssetData MInstanceData = AssetRegistry.GetAssetByObjectPath(FName(*MInstancePath));
+		FAssetData MInstanceData = AssetRegistry.GetAssetByObjectPath(FSoftObjectPath(MInstancePath));
 
 		if (!MInstanceData.IsValid())
 		{
@@ -165,7 +165,7 @@ void FImportProgressiveSurfaces::ImportAsset(TSharedPtr<FJsonObject> AssetImport
 			}
 		}		
 
-		FAssetData TextureData = AssetRegistry.GetAssetByObjectPath(FName(*TexturePath));
+		FAssetData TextureData = AssetRegistry.GetAssetByObjectPath(FSoftObjectPath(TexturePath));
 
 		if (!TextureData.IsValid()) return;
 
@@ -186,7 +186,7 @@ void FImportProgressiveSurfaces::ImportAsset(TSharedPtr<FJsonObject> AssetImport
 			}
 		}		
 
-		FAssetData NormalData = AssetRegistry.GetAssetByObjectPath(FName(*NormalPath));
+		FAssetData NormalData = AssetRegistry.GetAssetByObjectPath(FSoftObjectPath(NormalPath));
 
 		if (!NormalData.IsValid()) return;
 
@@ -197,7 +197,7 @@ void FImportProgressiveSurfaces::ImportAsset(TSharedPtr<FJsonObject> AssetImport
 	else if (ImportData->ProgressiveStage == 4)
 	{	
 		FString MInstanceHighPath = AssetMetaData.materialInstances[0].instancePath;
-		FAssetData MInstanceHighData = AssetRegistry.GetAssetByObjectPath(FName(*MInstanceHighPath));
+		FAssetData MInstanceHighData = AssetRegistry.GetAssetByObjectPath(FSoftObjectPath(MInstanceHighPath));
 
 		if (!MInstanceHighData.IsValid()) return;
 
@@ -329,7 +329,7 @@ void FImportProgressiveSurfaces::HandleNormalMaterialLoad(FAssetData AssetInstan
 
 		FAssetRegistryModule& AssetRegistryModule = FModuleManager::GetModuleChecked<FAssetRegistryModule>("AssetRegistry");
 		IAssetRegistry& AssetRegistry = AssetRegistryModule.Get();
-		FAssetData OverridenInstanceData = AssetRegistry.GetAssetByObjectPath(FName(*OverridenInstance->GetPathName()));
+		FAssetData OverridenInstanceData = AssetRegistry.GetAssetByObjectPath(FSoftObjectPath(OverridenInstance));
 		SpawnMaterialPreviewActor(AssetMetaData.assetID, LocationOffset, true, OverridenInstanceData);
 		return;
 	}
diff --git a/Engine/Plugins/Bridge/Source/MegascansPlugin/Private/Tools/BlendMaterials.cpp b/Engine/Plugins/Bridge/Source/MegascansPlugin/Private/Tools/BlendMaterials.cpp
--- a/Engine/Plugins/Bridge/Source/MegascansPlugin/Private/Tools/BlendMaterials.cpp
+++ b/Engine/Plugins/Bridge/Source/MegascansPlugin/Private/Tools/BlendMaterials.cpp
@@ -151,7 +151,7 @@ void FMaterialBlend::ConvertToVirtualTextures(FUAssetMeta AssetMetaData)
 
     for (FTexturesList TextureMeta : AssetMetaData.textureSets)
     {
-        FAssetData TextureData = AssetRegistry.GetAssetByObjectPath(FName(*TextureMeta.path));
+        FAssetData TextureData = AssetRegistry.GetAssetByObjectPath(FSoftObjectPath(TextureMeta.path));
 
         if (!TextureData.IsValid())
             return;
diff --git a/Engine/Plugins/Bridge/Source/MegascansPlugin/Private/Utilities/MaterialUtils.cpp b/Engine/Plugins/Bridge/Source/MegascansPlugin/Private/Utilities/MaterialUtils.cpp
--- a/Engine/Plugins/Bridge/Source/MegascansPlugin/Private/Utilities/MaterialUtils.cpp
+++ b/Engine/Plugins/Bridge/Source/MegascansPlugin/Private/Utilities/MaterialUtils.cpp
@@ -93,7 +93,7 @@ UMaterialInstanceConstant *FMaterialUtils::CreateMaterialOverride(FUAssetMeta As
         if (UMaterialEditingLibrary::GetMaterialInstanceTextureParameterValue(OverrideInstance, FName(*Texture.type)))
         {
 
-            FAssetData TextureData = AssetRegistry.GetAssetByObjectPath(FName(*Texture.path));
+            FAssetData TextureData = AssetRegistry.GetAssetByObjectPath(FSoftObjectPath(Texture.path));
 
             UTexture *TextureAsset = Cast<UTexture>(TextureData.GetAsset());
             if (UMaterialEditingLibrary::SetMaterialInstanceTextureParameterValue(OverrideInstance, FName(*Texture.type), TextureAsset))
@@ -115,7 +115,7 @@ void FMaterialUtils::ApplyMaterialInstance(FUAssetMeta AssetMetaData, UMaterialI
 
     for (FMeshInfo MeshMeta : AssetMetaData.meshList)
     {
-        FAssetData MeshData = AssetRegistry.GetAssetByObjectPath(FName(*MeshMeta.path));
+        FAssetData MeshData = AssetRegistry.GetAssetByObjectPath(FSoftObjectPath(MeshMeta.path));
         UStaticMesh *AssetMesh = Cast<UStaticMesh>(MeshData.GetAsset());
         AssetMesh->SetMaterial(0, CastChecked<UMaterialInterface>(MaterialInstance));
 
@@ -150,7 +150,7 @@ bool FMaterialUtils::ShouldOverrideMaterial(const FString &AssetType)
 TArray<AStaticMeshActor *> FMaterialUtils::ApplyMaterialToSelection(const FString &InstancePath)
 {
     IAssetRegistry &AssetRegistry = FModuleManager::GetModuleChecked<FAssetRegistryModule>("AssetRegistry").Get();
-    FAssetData InstanceData = AssetRegistry.GetAssetByObjectPath(FName(*InstancePath));
+    FAssetData InstanceData = AssetRegistry.GetAssetByObjectPath(FSoftObjectPath(InstancePath));
     UMaterialInstanceConstant *MaterialInstance = Cast<UMaterialInstanceConstant>(UEditorAssetLibrary::LoadAsset(InstancePath));
 
     USelection *SelectedActors = GEditor->GetSelectedActors();
diff --git a/Engine/Plugins/Bridge/Source/MegascansPlugin/Private/Utilities/MiscUtils.cpp b/Engine/Plugins/Bridge/Source/MegascansPlugin/Private/Utilities/MiscUtils.cpp
--- a/Engine/Plugins/Bridge/Source/MegascansPlugin/Private/Utilities/MiscUtils.cpp
+++ b/Engine/Plugins/Bridge/Source/MegascansPlugin/Private/Utilities/MiscUtils.cpp
@@ -360,10 +360,10 @@ void AssetUtils::SyncFolder(const FString& TargetFolder)
 	// Testing syncing
 }
 
-void AssetUtils::RegisterAsset(const FString& PackagePath)
+void AssetUtils::RegisterAsset(const FString& PackagePath) // Should this be named AssetPath?
 {
 	IAssetRegistry& AssetRegistry = FModuleManager::GetModuleChecked<FAssetRegistryModule>("AssetRegistry").Get();
-	FAssetData CharacterAssetData = AssetRegistry.GetAssetByObjectPath(FName(*PackagePath));
+	FAssetData CharacterAssetData = AssetRegistry.GetAssetByObjectPath(FSoftObjectPath(PackagePath));
 	UObject* CharacterObject = CharacterAssetData.GetAsset();
 }
 
