Additional FName -> FSoftObjectPath deprecation fixups

diff --git a/Engine/Plugins/Bridge/Source/MegascansPlugin/Private/Utilities/MiscUtils.cpp b/Engine/Plugins/Bridge/Source/MegascansPlugin/Private/Utilities/MiscUtils.cpp
--- a/Engine/Plugins/Bridge/Source/MegascansPlugin/Private/Utilities/MiscUtils.cpp
+++ b/Engine/Plugins/Bridge/Source/MegascansPlugin/Private/Utilities/MiscUtils.cpp
@@ -211,7 +211,7 @@ void AssetUtils::DeleteDirectory(FString TargetDirectory)
 bool AssetUtils::DeleteAsset(const FString& AssetPath)
 {
 	IAssetRegistry& AssetRegistry = FModuleManager::GetModuleChecked<FAssetRegistryModule>("AssetRegistry").Get();
-	FAssetData AssetData = AssetRegistry.GetAssetByObjectPath(FName(*AssetPath));
+	FAssetData AssetData = AssetRegistry.GetAssetByObjectPath(FSoftObjectPath(AssetPath));
 
 	if (AssetData.IsAssetLoaded())
 	{
@@ -256,7 +256,7 @@ TArray<UMaterialInstanceConstant*> AssetUtils::GetSelectedAssets(const FTopLevel
 	{
 		if (SelectedAsset.AssetClassPath == AssetClass)
 		{
-			ObjectArray.Add(CastChecked<UMaterialInstanceConstant>(UEditorAssetLibrary::LoadAsset(SelectedAsset.ObjectPath.ToString())));
+			ObjectArray.Add(CastChecked<UMaterialInstanceConstant>(UEditorAssetLibrary::LoadAsset(SelectedAsset.GetObjectPathString())));
 		}
 	}
 
@@ -269,7 +269,7 @@ void AssetUtils::AddFoliageTypesToLevel(TArray<FString> FoliageTypePaths)
 
 	for (FString FoliageTypePath : FoliageTypePaths)
 	{
-		FAssetData AssetData = AssetRegistry.GetAssetByObjectPath(FName(*FoliageTypePath));
+		FAssetData AssetData = AssetRegistry.GetAssetByObjectPath(FSoftObjectPath(FoliageTypePath));
 		if (!AssetData.IsValid()) return;
 
 		auto* CurrentWorld = GEditor->GetEditorWorldContext().World();
diff --git a/Engine/Source/Editor/ContentBrowserData/Private/ContentBrowserDataSource.cpp b/Engine/Source/Editor/ContentBrowserData/Private/ContentBrowserDataSource.cpp
--- a/Engine/Source/Editor/ContentBrowserData/Private/ContentBrowserDataSource.cpp
+++ b/Engine/Source/Editor/ContentBrowserData/Private/ContentBrowserDataSource.cpp
@@ -230,7 +230,9 @@ TArray<FContentBrowserItemPath> UContentBrowserDataSource::GetAliasesForPath(con
 
 TArray<FContentBrowserItemPath> UContentBrowserDataSource::GetAliasesForPath(FName InInternalPath) const
 {
+PRAGMA_DISABLE_DEPRECATION_WARNINGS
 	return GetAliasesForPath(FSoftObjectPath(InInternalPath));
+PRAGMA_ENABLE_DEPRECATION_WARNINGS
 }
 
 bool UContentBrowserDataSource::IsDiscoveringItems(FText* OutStatus)
diff --git a/Engine/Source/Editor/SourceControlWindowExtender/Private/SourceControlWindowExtenderModule.cpp b/Engine/Source/Editor/SourceControlWindowExtender/Private/SourceControlWindowExtenderModule.cpp
--- a/Engine/Source/Editor/SourceControlWindowExtender/Private/SourceControlWindowExtenderModule.cpp
+++ b/Engine/Source/Editor/SourceControlWindowExtender/Private/SourceControlWindowExtenderModule.cpp
@@ -88,7 +88,7 @@ void FSourceControlWindowExtenderModule::GetAssetsFromFilenames(const TArray<FSt
 			const FAssetData& AssetData = OutAssets[0];
 			if (TSubclassOf<AActor> ActorClass = AssetData.GetClass())
 			{
-				if (CurrentWorld && AssetData.ObjectPath.ToString().StartsWith(CurrentWorld->GetPathName()))
+				if (CurrentWorld && AssetData.GetObjectPathString().StartsWith(CurrentWorld->GetPathName()))
 				{
 					if (AssetData.IsAssetLoaded())
 					{
diff --git a/Engine/Source/Runtime/Engine/Private/PrimaryAssetLabel.cpp b/Engine/Source/Runtime/Engine/Private/PrimaryAssetLabel.cpp
--- a/Engine/Source/Runtime/Engine/Private/PrimaryAssetLabel.cpp
+++ b/Engine/Source/Runtime/Engine/Private/PrimaryAssetLabel.cpp
@@ -68,7 +68,9 @@ void UPrimaryAssetLabel::UpdateAssetBundleData()
 		TArray<FName> CollectionAssetPathNames;
 		ICollectionManager& CollectionManager = FCollectionManagerModule::GetModule().Get();
 		CollectionManager.GetAssetsInCollection(AssetCollection.CollectionName, ECollectionShareType::CST_All, CollectionAssetPathNames);
+PRAGMA_DISABLE_DEPRECATION_WARNINGS
 		TArray<FSoftObjectPath> CollectionAssets = UE::SoftObjectPath::Private::ConvertObjectPathNames(CollectionAssetPathNames);
+PRAGMA_ENABLE_DEPRECATION_WARNINGS
 		for (int32 Index = 0; Index < CollectionAssets.Num(); ++Index)
 		{
 			FAssetData FoundAsset = Manager.GetAssetRegistry().GetAssetByObjectPath(CollectionAssets[Index]);
