MegascansPlugin - fix bad automerge.
#rnx

diff --git a/Engine/Plugins/Bridge/Source/MegascansPlugin/Private/Utilities/MiscUtils.cpp b/Engine/Plugins/Bridge/Source/MegascansPlugin/Private/Utilities/MiscUtils.cpp
--- a/Engine/Plugins/Bridge/Source/MegascansPlugin/Private/Utilities/MiscUtils.cpp
+++ b/Engine/Plugins/Bridge/Source/MegascansPlugin/Private/Utilities/MiscUtils.cpp
@@ -255,6 +255,7 @@ TArray<UMaterialInstanceConstant*> AssetUtils::GetSelectedAssets(const FTopLevel
 	for (FAssetData SelectedAsset : AssetDatas)
 	{
 		if (SelectedAsset.AssetClassPath == AssetClass)
+		{
 			ObjectArray.Add(CastChecked<UMaterialInstanceConstant>(UEditorAssetLibrary::LoadAsset(SelectedAsset.GetObjectPathString())));
 		}
 
