Fixing Bridge missing the ThirdParty folder for linux

#rnx

diff --git a/Engine/Plugins/Bridge/Config/FilterPluginLinux.ini b/Engine/Plugins/Bridge/Config/FilterPluginLinux.ini
new file mode 100644
--- /dev/null
+++ b/Engine/Plugins/Bridge/Config/FilterPluginLinux.ini
@@ -0,0 +1,2 @@
+[FilterPlugin]
+/ThirdParty/Linux/...
\ No newline at end of file
