[Megascans]
* Removed unused enum that was causing compile warnings due to enum value "normal" being defined in global namespace and makes other enums shadowed

diff --git a/Engine/Plugins/Bridge/Source/MegascansPlugin/Private/Tools/BlendMaterials.h b/Engine/Plugins/Bridge/Source/MegascansPlugin/Private/Tools/BlendMaterials.h
--- a/Engine/Plugins/Bridge/Source/MegascansPlugin/Private/Tools/BlendMaterials.h
+++ b/Engine/Plugins/Bridge/Source/MegascansPlugin/Private/Tools/BlendMaterials.h
@@ -6,13 +6,6 @@
 struct FAssetData;
 struct FUAssetMeta;
 
-enum EImportState
-{
-    Instance,
-    Albedo,
-    Normal
-};
-
 class FMaterialBlend
 {
 private:
