Quixel Bridge - fix mac os crash on login when the user types in email/password

diff --git a/Engine/Plugins/Bridge/Source/Bridge/Private/UI/BridgeUIManager.cpp b/Engine/Plugins/Bridge/Source/Bridge/Private/UI/BridgeUIManager.cpp
--- a/Engine/Plugins/Bridge/Source/Bridge/Private/UI/BridgeUIManager.cpp
+++ b/Engine/Plugins/Bridge/Source/Bridge/Private/UI/BridgeUIManager.cpp
@@ -265,6 +265,15 @@ TSharedRef<SDockTab> FBridgeUIManagerImpl::CreateBridgeTab(const FSpawnTabArgs&
 
 	TSharedPtr<SWebBrowser> PluginWebBrowser;
 
+#if PLATFORM_MAC
+	WindowSettings.InitialURL = FinalUrl;
+	IWebBrowserSingleton* WebBrowserSingleton = IWebBrowserModule::Get().GetSingleton();
+	Browser = WebBrowserSingleton->CreateBrowserWindow(WindowSettings);
+	PluginWebBrowser = SAssignNew(WebBrowserWidget, SWebBrowser, Browser)
+		.ShowAddressBar(false)
+		.ShowControls(false);
+
+#elif PLATFORM_WINDOWS || PLATFORM_LINUX
 	FWebBrowserInitSettings browserInitSettings = FWebBrowserInitSettings();
 	IWebBrowserModule::Get().CustomInitialize(browserInitSettings);
 	WindowSettings.InitialURL = FinalUrl;
@@ -285,6 +294,7 @@ TSharedRef<SDockTab> FBridgeUIManagerImpl::CreateBridgeTab(const FSpawnTabArgs&
 		return SAssignNew(LocalBrowserDock, SDockTab)
 			.TabRole(ETabRole::NomadTab);
 	}
+#endif
 
 	SAssignNew(LocalBrowserDock, SDockTab)
 		.OnTabClosed_Lambda([](TSharedRef<class SDockTab> InParentTab)
