- Stop keypresses in the Quixel Bridge web view from being passed along to the editor

diff --git a/Engine/Plugins/Bridge/Source/Bridge/Private/UI/BridgeUIManager.cpp b/Engine/Plugins/Bridge/Source/Bridge/Private/UI/BridgeUIManager.cpp
--- a/Engine/Plugins/Bridge/Source/Bridge/Private/UI/BridgeUIManager.cpp
+++ b/Engine/Plugins/Bridge/Source/Bridge/Private/UI/BridgeUIManager.cpp
@@ -10,6 +10,7 @@
 #include "SWebBrowser.h"
 #include "Serialization/JsonSerializer.h"
 #include "WebBrowserModule.h"
+#include "IWebBrowserWindow.h"
 // Widgets
 #include "Framework/Application/SlateApplication.h"
 #include "ToolMenu.h"
@@ -296,6 +297,12 @@ TSharedRef<SDockTab> FBridgeUIManagerImpl::CreateBridgeTab(const FSpawnTabArgs&
 	}
 #endif
 
+	// by handling these keys, keys that were swallowed as Char events won't come out as unhandled KwyDown/Up events
+	// that then go to the editor and change editor state (these arent needed if the function that calls these callbacks
+	// was deleted, which may be the case soon)
+	Browser->OnUnhandledKeyUp().BindLambda([](const FKeyEvent&) { return true; });
+	Browser->OnUnhandledKeyDown().BindLambda([](const FKeyEvent&) { return true; });
+
 	SAssignNew(LocalBrowserDock, SDockTab)
 		.OnTabClosed_Lambda([](TSharedRef<class SDockTab> InParentTab)
 		{
