Quixel Bridge - Fix for surfaces d&d and decal (veggecd) crash

diff --git a/Engine/Plugins/Bridge/Source/Bridge/Private/UI/BrowserBinding.cpp b/Engine/Plugins/Bridge/Source/Bridge/Private/UI/BrowserBinding.cpp
--- a/Engine/Plugins/Bridge/Source/Bridge/Private/UI/BrowserBinding.cpp
+++ b/Engine/Plugins/Bridge/Source/Bridge/Private/UI/BrowserBinding.cpp
@@ -393,6 +393,11 @@ void SetupOnActorsDroppedEvent()
 			Filtered3dObjects = Filtered3dObjects.FilterByPredicate([](UObject* Object) { return Object != nullptr; });
 			Filtered3dActors = Filtered3dActors.FilterByPredicate([](AActor* Actor) { return Actor != nullptr; });
 
+			if (Filtered3dIDs.Num() != Filtered3dObjects.Num() || Filtered3dObjects.Num() != Filtered3dActors.Num())
+			{
+				return;
+			}
+			
 			for (int32 i = 0; i < Filtered3dObjects.Num(); i++)
 			{
 				UObject* Object = Filtered3dObjects[i];
diff --git a/Engine/Plugins/Bridge/Source/MegascansPlugin/Private/AssetImporters/ProgressiveImportSurfaces.h b/Engine/Plugins/Bridge/Source/MegascansPlugin/Private/AssetImporters/ProgressiveImportSurfaces.h
--- a/Engine/Plugins/Bridge/Source/MegascansPlugin/Private/AssetImporters/ProgressiveImportSurfaces.h
+++ b/Engine/Plugins/Bridge/Source/MegascansPlugin/Private/AssetImporters/ProgressiveImportSurfaces.h
@@ -10,10 +10,9 @@ struct FUAssetMeta;
 struct FProgressiveSurfaces
 {	
 	TArray<AStaticMeshActor*> ActorsInLevel;
-	UMaterialInstanceConstant* PreviewInstance;
+	UMaterialInstanceConstant* PreviewInstance = nullptr;
 	FString PreviewFolderPath;
 	FString PreviewMeshPath;
-
 };
 
 class FImportProgressiveSurfaces
