Fix missing FGCObject referencer name

#rnx

diff --git a/Engine/Plugins/Bridge/Source/MegascansPlugin/Private/UI/SMSWindow.cpp b/Engine/Plugins/Bridge/Source/MegascansPlugin/Private/UI/SMSWindow.cpp
--- a/Engine/Plugins/Bridge/Source/MegascansPlugin/Private/UI/SMSWindow.cpp
+++ b/Engine/Plugins/Bridge/Source/MegascansPlugin/Private/UI/SMSWindow.cpp
@@ -154,6 +154,8 @@ private:
 	UMegascansSettings* MegascansSettings;
 	UMaterialBlendSettings* MaterialBlendSettings;
 	UMaterialPresetsSettings* MaterialOverrideSettings;
+
+	virtual FString GetReferencerName() const override { return "MegascanSettings"; }
 };
 	
 
