Bifrost: Fix non-unity build for MegaScans plugin

#rnx

diff --git a/Engine/Plugins/Bridge/Source/MegascansPlugin/Private/AssetsImportController.cpp b/Engine/Plugins/Bridge/Source/MegascansPlugin/Private/AssetsImportController.cpp
--- a/Engine/Plugins/Bridge/Source/MegascansPlugin/Private/AssetsImportController.cpp
+++ b/Engine/Plugins/Bridge/Source/MegascansPlugin/Private/AssetsImportController.cpp
@@ -5,6 +5,7 @@
 #include "AssetImporters/ProgressiveImport3D.h"
 #include "AssetImporters/ProgressiveImportSurfaces.h"
 #include "MetaHumanProjectUtilities.h"
+#include "Misc/Paths.h"
 
 
 #include "Dom/JsonObject.h"
