Quixel Bridge - remove target platforms from uplugin file, which possibly fixes cook errors on other platforms

diff --git a/Engine/Plugins/Bridge/Bridge.uplugin b/Engine/Plugins/Bridge/Bridge.uplugin
--- a/Engine/Plugins/Bridge/Bridge.uplugin
+++ b/Engine/Plugins/Bridge/Bridge.uplugin
@@ -12,11 +12,6 @@
 	"SupportURL": "[REDACTED_EMAIL]",
 	"CanContainContent": true,
 	"EnabledByDefault": true,
-	"SupportedTargetPlatforms": [
-		"Win64",
-		"Mac",
-		"Linux"
-	],
 	"Modules": [
 		{
 			"Name": "Bridge",
