Fix a crash that would occur if the user dragged an asset to the viewport after closing the bridge panel.

- UBrowserBinding was being registered for callbacks from the editor on various mouse drag events but was not being unregistered from these callbacks when the window was destroyed. Invoking these events after this point would cause invalid memory access.
- Since UBrowserBinding is a UObject we can use the UObject style of binding to the editor delegate which will retain a weakptr to the object and clean it up automatically after it is deleted so we don't have to worry about unregistering ourselves.
- While trying to reproduce the main issue I stumbled on a one off crash in UBrowserBinding::SwitchDragDropOp where the top level window was not set. I think this might have been related to my use of multiple monitors at the time (I think I dragged between them) and have been unable to reproduce it since switching to a single monitor setup. I added a speculative fix for this has the crash has been seen in the wild but with a very low repro rate.

diff --git a/Engine/Plugins/Bridge/Source/Bridge/Private/UI/BrowserBinding.cpp b/Engine/Plugins/Bridge/Source/Bridge/Private/UI/BrowserBinding.cpp
--- a/Engine/Plugins/Bridge/Source/Bridge/Private/UI/BrowserBinding.cpp
+++ b/Engine/Plugins/Bridge/Source/Bridge/Private/UI/BrowserBinding.cpp
@@ -351,117 +351,119 @@ void UBrowserBinding::SwitchDragDropOp(TArray<FString> URLs, TSharedRef<FAssetDr
 			0,
 			ModifierKeyState);
 
-	// Tell slate to enter drag and drop mode.
-	// Make a fake mouse event for slate, so we can initiate a drag and drop.
-	FDragDropEvent DragDropEvent(FakePointerEvent, DragDropOperation);
+	if (TSharedPtr<SWindow> TopLevelWindow = FSlateApplication::Get().GetActiveTopLevelWindow(); TopLevelWindow.IsValid())
+	{
+		// Tell slate to enter drag and drop mode.
+		// Make a fake mouse event for slate, so we can initiate a drag and drop.
+		FDragDropEvent DragDropEvent(FakePointerEvent, DragDropOperation);
 
-	FSlateApplication::Get().ProcessDragEnterEvent(FSlateApplication::Get().GetActiveTopLevelWindow().ToSharedRef(), DragDropEvent);
+		FSlateApplication::Get().ProcessDragEnterEvent(TopLevelWindow.ToSharedRef(), DragDropEvent);
+	}
 }
 
-void SetupOnActorsDroppedEvent()
+void UBrowserBinding::OnNewActorsDropped(const TArray<UObject*>& DroppedObjects, const TArray<AActor*>& DroppedActors)
 {
-	FEditorDelegates::OnNewActorsDropped.AddLambda([](TArray<UObject*> DroppedObjects, TArray<AActor*> DroppedActors)
+	if (bIsDragging)
 	{
-		if (FBridgeUIManager::BrowserBinding->bIsDragging)
-		{
-			// Take out 3d from the list of DroppedObjects
-			TArray<UObject*> Filtered3dObjects;
-			TArray<AActor*> Filtered3dActors;
-			TArray<FString> Filtered3dIDs;
-
-			for (int32 i = 0; i < DroppedObjects.Num(); i++)
-			{	
-				if (DroppedObjects[i]->IsA(UStaticMesh::StaticClass()) || DroppedObjects[i]->IsA(UMaterialInstanceConstant::StaticClass()))
-				{
-					Filtered3dObjects.Add(DroppedObjects[i]);
-				}
-			}
+		// Take out 3d from the list of DroppedObjects
+		TArray<UObject*> Filtered3dObjects;
+		TArray<AActor*> Filtered3dActors;
+		TArray<FString> Filtered3dIDs;
 
-			// Take out actors other than decals from the list of DroppedActors
-			for (int32 i = 0; i < DroppedActors.Num(); i++)
-			{	
-				if (DroppedActors[i]->IsA(AStaticMeshActor::StaticClass()) || DroppedActors[i]->IsA(ADecalActor::StaticClass()))
-				{
-					Filtered3dActors.Add(DroppedActors[i]);
-				}
-			}
-
-			for (int32 i = 0; i < FBridgeUIManager::BrowserBinding->DragDropTypes.Num(); i++)
+		for (int32 i = 0; i < DroppedObjects.Num(); i++)
+		{
+			if (DroppedObjects[i]->IsA(UStaticMesh::StaticClass()) || DroppedObjects[i]->IsA(UMaterialInstanceConstant::StaticClass()))
 			{
-				FString Type = FBridgeUIManager::BrowserBinding->DragDropTypes[i];
-				FString ID = FBridgeUIManager::BrowserBinding->DragDropIDs[i];
-				if (Type == "surface" || Type == "imperfection") continue;
-				Filtered3dIDs.Add(ID);
+				Filtered3dObjects.Add(DroppedObjects[i]);
 			}
+		}
 
-			if (Filtered3dActors.Num() != Filtered3dObjects.Num())
+		// Take out actors other than decals from the list of DroppedActors
+		for (int32 i = 0; i < DroppedActors.Num(); i++)
+		{
+			if (DroppedActors[i]->IsA(AStaticMeshActor::StaticClass()) || DroppedActors[i]->IsA(ADecalActor::StaticClass()))
 			{
-				return;
+				Filtered3dActors.Add(DroppedActors[i]);
 			}
+		}
+
+		for (int32 i = 0; i < DragDropTypes.Num(); i++)
+		{
+			FString Type = DragDropTypes[i];
+			FString ID = DragDropIDs[i];
+			if (Type == "surface" || Type == "imperfection") continue;
+			Filtered3dIDs.Add(ID);
+		}
 
-			// Search for the ID against filtered objects (specific case for decals)
-			for (int32 i = 0; i < Filtered3dIDs.Num(); i++)
+		if (Filtered3dActors.Num() != Filtered3dObjects.Num())
+		{
+			return;
+		}
+
+		// Search for the ID against filtered objects (specific case for decals)
+		for (int32 i = 0; i < Filtered3dIDs.Num(); i++)
+		{
+			FString ID = Filtered3dIDs[i];
+			for (int32 j = 0; j < Filtered3dObjects.Num(); j++)
 			{
-				FString ID = Filtered3dIDs[i];
-				for (int32 j = 0; j < Filtered3dObjects.Num(); j++)
+				UObject* Object = Filtered3dObjects[j];
+				AActor* Actor = Filtered3dActors[j];
+				if (Object != nullptr && Object->GetName().Find(ID) != -1)
 				{
-					UObject* Object = Filtered3dObjects[j];
-					AActor* Actor = Filtered3dActors[j];
-					if (Object != nullptr && Object->GetName().Find(ID) != -1)
-					{
-						FBridgeUIManager::BrowserBinding->AssetToSphereMap.Add(ID, Actor);
-						Filtered3dIDs[i] = "-1";
-						Filtered3dObjects[j] = nullptr;
-						Filtered3dActors[j] = nullptr;
-					}
+					AssetToSphereMap.Add(ID, Actor);
+					Filtered3dIDs[i] = "-1";
+					Filtered3dObjects[j] = nullptr;
+					Filtered3dActors[j] = nullptr;
 				}
 			}
+		}
 
-			Filtered3dIDs = Filtered3dIDs.FilterByPredicate([](FString ID) { return ID != "-1"; });
-			Filtered3dObjects = Filtered3dObjects.FilterByPredicate([](UObject* Object) { return Object != nullptr; });
-			Filtered3dActors = Filtered3dActors.FilterByPredicate([](AActor* Actor) { return Actor != nullptr; });
+		Filtered3dIDs = Filtered3dIDs.FilterByPredicate([](FString ID) { return ID != "-1"; });
+		Filtered3dObjects = Filtered3dObjects.FilterByPredicate([](UObject* Object) { return Object != nullptr; });
+		Filtered3dActors = Filtered3dActors.FilterByPredicate([](AActor* Actor) { return Actor != nullptr; });
 
-			if (Filtered3dIDs.Num() != Filtered3dObjects.Num() || Filtered3dObjects.Num() != Filtered3dActors.Num())
-			{
-				return;
-			}
-			
-			for (int32 i = 0; i < Filtered3dObjects.Num(); i++)
-			{
-				UObject* Object = Filtered3dObjects[i];
-				AActor* Actor = Filtered3dActors[i];
-				FString ID = Filtered3dIDs[i];
+		if (Filtered3dIDs.Num() != Filtered3dObjects.Num() || Filtered3dObjects.Num() != Filtered3dActors.Num())
+		{
+			return;
+		}
 
-				if (Object->GetName().Find("Sphere") != -1)
-				{
-					FBridgeUIManager::BrowserBinding->AssetToSphereMap.Add(ID, Actor);
-				}
+		for (int32 i = 0; i < Filtered3dObjects.Num(); i++)
+		{
+			UObject* Object = Filtered3dObjects[i];
+			AActor* Actor = Filtered3dActors[i];
+			FString ID = Filtered3dIDs[i];
+
+			if (Object->GetName().Find("Sphere") != -1)
+			{
+				AssetToSphereMap.Add(ID, Actor);
 			}
 		}
-	});
+	}
 }
 
-void SetupOnObjectAppliedToActorEvent()
+void UBrowserBinding::OnApplyObjectToActor(UObject* DroppedObj, AActor* TargetActor)
 {
-	FEditorDelegates::OnApplyObjectToActor.AddLambda([](UObject* TheObject, AActor* Actor)
+	if (!bIsDragging)
+	{
+		return;
+	}
+
+	for (int32 i = 0; i < DragDropTypes.Num(); i++)
 	{
-		if (!FBridgeUIManager::BrowserBinding->bIsDragging) return;
-		for (int32 i = 0; i < FBridgeUIManager::BrowserBinding->DragDropTypes.Num(); i++)
+		FString Type = DragDropTypes[i];
+		if (Type == "surface" || Type == "imperfection")
 		{
-			FString Type = FBridgeUIManager::BrowserBinding->DragDropTypes[i];
-			if (Type == "surface" || Type == "imperfection")
+			if (FBridgeDragDropHelper::Instance->SurfaceToActorMap.Contains(DragDropIDs[i]))
 			{
-				if (FBridgeDragDropHelper::Instance->SurfaceToActorMap.Contains(FBridgeUIManager::BrowserBinding->DragDropIDs[i]))
-				{
-					FBridgeDragDropHelper::Instance->SurfaceToActorMap.Remove(FBridgeUIManager::BrowserBinding->DragDropIDs[i]);
-				}
-
-				FBridgeDragDropHelper::Instance->SurfaceToActorMap.Add(FBridgeUIManager::BrowserBinding->DragDropIDs[i], Actor);
+				FBridgeDragDropHelper::Instance->SurfaceToActorMap.Remove(DragDropIDs[i]);
 			}
+
+			FBridgeDragDropHelper::Instance->SurfaceToActorMap.Add(DragDropIDs[i], TargetActor);
 		}
-	});
+	}
 }
 
+
 void PopuplateInAssetData()
 {
 	TArray<FString> Types = FBridgeUIManager::BrowserBinding->DragDropTypes;
@@ -563,8 +565,9 @@ void UBrowserBinding::DragStarted(TArray<FString> ImageUrls, TArray<FString> IDs
 
 	if (!bIsDropEventBound)
 	{
-		SetupOnActorsDroppedEvent();
-		SetupOnObjectAppliedToActorEvent();
+		FEditorDelegates::OnNewActorsDropped.AddUObject(FBridgeUIManager::BrowserBinding, &UBrowserBinding::OnNewActorsDropped);
+		FEditorDelegates::OnApplyObjectToActor.AddUObject(FBridgeUIManager::BrowserBinding, &UBrowserBinding::OnApplyObjectToActor);
+
 		FBridgeUIManager::BrowserBinding->bIsDropEventBound = true;
 	}
 
diff --git a/Engine/Plugins/Bridge/Source/Bridge/Private/UI/BrowserBinding.h b/Engine/Plugins/Bridge/Source/Bridge/Private/UI/BrowserBinding.h
--- a/Engine/Plugins/Bridge/Source/Bridge/Private/UI/BrowserBinding.h
+++ b/Engine/Plugins/Bridge/Source/Bridge/Private/UI/BrowserBinding.h
@@ -27,6 +27,9 @@ class UBrowserBinding : public UObject
 private:
 	void SwitchDragDropOp(TArray<FString> URLs, TSharedRef<FAssetDragDropOp> DragDropOperation);
 	
+	void OnNewActorsDropped(const TArray<UObject*>& DroppedObjects, const TArray<AActor*>& DroppedActors);
+	void OnApplyObjectToActor(UObject* DroppedObj, AActor* TargetActor);
+
 public:
 	FOnDialogSuccess DialogSuccessDelegate;
 	FOnDialogFail DialogFailDelegate;
