Fix various warnings generated when parsing Json while importing a MetaHuman

diff --git a/Engine/Plugins/Bridge/Source/MegascansPlugin/Private/AssetsImportController.cpp b/Engine/Plugins/Bridge/Source/MegascansPlugin/Private/AssetsImportController.cpp
--- a/Engine/Plugins/Bridge/Source/MegascansPlugin/Private/AssetsImportController.cpp
+++ b/Engine/Plugins/Bridge/Source/MegascansPlugin/Private/AssetsImportController.cpp
@@ -42,7 +42,6 @@ void FAssetsImportController::DataReceived(const FString DataFromBridge)
 	for (const TSharedPtr<FJsonValue>& AssetJson : AssetsImportDataArray)
 	{
 		EAssetImportType ImportType = JsonUtils::GetImportType(AssetJson->AsObject());
-		FString AssetType = AssetJson->AsObject()->GetStringField(TEXT("assetType"));
 		if (ImportType == EAssetImportType::DHI_CHARACTER)
 		{
 			BatchImportCharacters.Add(AssetJson->AsObject()->GetStringField(TEXT("folderName")));
@@ -52,11 +51,11 @@ void FAssetsImportController::DataReceived(const FString DataFromBridge)
 	for (const TSharedPtr<FJsonValue>& AssetJson : AssetsImportDataArray)
 	{
 		const EAssetImportType ImportType = JsonUtils::GetImportType(AssetJson->AsObject());
-		const FString AssetType = AssetJson->AsObject()->GetStringField(TEXT("assetType"));
 		const FString ExportMode = AssetJson->AsObject()->GetStringField(TEXT("exportMode"));
 
 		if (ImportType == EAssetImportType::MEGASCANS_UASSET)
 		{
+			const FString AssetType = AssetJson->AsObject()->GetStringField(TEXT("assetType"));
 			CopyMSPresets();
 			//if (!SupportedAssetTypes.Contains(AssetType)) continue;
 
@@ -94,12 +93,17 @@ void FAssetsImportController::DataReceived(const FString DataFromBridge)
 			}
 
 			const FString CharacterName = AssetJsonObject->GetStringField(TEXT("folderName"));
+
+			// Value is only present when set to true
+			bool bIsBulkExported = BatchImportCharacters.Num() > 1;
+			AssetJsonObject->TryGetBoolField(TEXT("isBulkExported"), bIsBulkExported);
+
 			FMetaHumanImportDescription AssetImportDescription = {
 				FPaths::Combine(AssetJsonObject->GetStringField(TEXT("characterPath")), CharacterName),
 				AssetJsonObject->GetStringField(TEXT("commonPath")),
 				CharacterName,
 				AssetJsonObject->GetStringField(TEXT("id")),
-				AssetJsonObject->GetBoolField(TEXT("isBulkExported")) || BatchImportCharacters.Num() > 1
+				bIsBulkExported
 			};
 
 			const TArray<TSharedPtr<FJsonValue>>* AvailableMetahumans;
@@ -107,11 +111,17 @@ void FAssetsImportController::DataReceived(const FString DataFromBridge)
 			{
 				for (const TSharedPtr<FJsonValue>& Entry : *AvailableMetahumans)
 				{
+					TSharedPtr<FJsonObject> EntryJsonObject = Entry->AsObject();
+
+					// Value is only present when set to true
+					bool bIsLegacy = false;
+					EntryJsonObject->TryGetBoolField(TEXT("isLegacy"), bIsLegacy);
+
 					AssetImportDescription.AccountMetaHumans.Emplace(FQuixelAccountMetaHumanEntry{
-						Entry->AsObject()->GetStringField(TEXT("name")),
-						Entry->AsObject()->GetStringField(TEXT("id")),
-						Entry->AsObject()->GetBoolField(TEXT("isLegacy")),
-						Entry->AsObject()->GetStringField(TEXT("version")),
+						EntryJsonObject->GetStringField(TEXT("name")),
+						EntryJsonObject->GetStringField(TEXT("id")),
+						bIsLegacy,
+						EntryJsonObject->GetStringField(TEXT("version")),
 					});
 				}
 			}
