Fix to NodeProcess and LiveCoding passing NULL to CreateProc Parms

diff --git a/Engine/Plugins/Bridge/Source/Bridge/Private/NodeProcess.cpp b/Engine/Plugins/Bridge/Source/Bridge/Private/NodeProcess.cpp
--- a/Engine/Plugins/Bridge/Source/Bridge/Private/NodeProcess.cpp
+++ b/Engine/Plugins/Bridge/Source/Bridge/Private/NodeProcess.cpp
@@ -41,7 +41,7 @@ void FNodeProcessManager::StartNodeProcess()
 	const FString ProcessURL = GetProcessURL();
 
 #if PLATFORM_WINDOWS || PLATFORM_LINUX
-	NodeProcessHandle = FPlatformProcess::CreateProc(*ProcessURL, NULL, true, true, true, &OutProcessId, 0, NULL, NULL, NULL);
+	NodeProcessHandle = FPlatformProcess::CreateProc(*ProcessURL, TEXT(""), true, true, true, &OutProcessId, 0, NULL, NULL, NULL);
 
 	int32 ReturnCode;
 	FPlatformProcess::GetProcReturnCode(NodeProcessHandle, &ReturnCode);
diff --git a/Engine/Source/Developer/Windows/LiveCoding/Private/LiveCodingModule2.cpp b/Engine/Source/Developer/Windows/LiveCoding/Private/LiveCodingModule2.cpp
--- a/Engine/Source/Developer/Windows/LiveCoding/Private/LiveCodingModule2.cpp
+++ b/Engine/Source/Developer/Windows/LiveCoding/Private/LiveCodingModule2.cpp
@@ -90,7 +90,7 @@ uint32 StartBroker()
 
 	FString BrokerExe = LivePlusPlusPath / TEXT("LPP_Broker.exe");
 	uint32 ProcessId;
-	if (FPlatformProcess::CreateProc(*BrokerExe, NULL, true, true, true, &ProcessId, 0, NULL, NULL, NULL).IsValid())
+	if (FPlatformProcess::CreateProc(*BrokerExe, TEXT(""), true, true, true, &ProcessId, 0, NULL, NULL, NULL).IsValid())
 	{
 		return ProcessId;
 	}
