Fab/Quixel: Don't crash if the WebBrowserWindow fails to be created

Log the error and add an error message in the tab so it doesn't just stay empty.
Also replace the modal dialog in the module unavailable case with a similar log+error message, and adjusted the message as WebBrowserModule doesn't require a plugin to be available (that's WebBrowserWidget instead)

#rnx

#Synced-CL 49370686

diff --git a/Engine/Plugins/Bridge/Source/Bridge/Private/UI/BridgeUIManager.cpp b/Engine/Plugins/Bridge/Source/Bridge/Private/UI/BridgeUIManager.cpp
--- a/Engine/Plugins/Bridge/Source/Bridge/Private/UI/BridgeUIManager.cpp
+++ b/Engine/Plugins/Bridge/Source/Bridge/Private/UI/BridgeUIManager.cpp
@@ -17,10 +17,10 @@
 #include "ToolMenus.h"
 #include "ToolMenuEntry.h"
 #include "ToolMenuSection.h"
-#include "Misc/MessageDialog.h"
 #include "ContentBrowserDataMenuContexts.h"
 #include "UI/BrowserBinding.h"
 #include "Widgets/Docking/SDockTab.h"
+#include "Widgets/Images/SImage.h"
 
 #define LOCTEXT_NAMESPACE "Bridge"
 #define LEVELEDITOR_MODULE_NAME TEXT("LevelEditor")
@@ -286,6 +286,31 @@ TSharedRef<SDockTab> FBridgeUIManagerImpl::CreateBridgeTab(const FSpawnTabArgs&
 	{
 		IWebBrowserSingleton* WebBrowserSingleton = IWebBrowserModule::Get().GetSingleton();
 		Browser = WebBrowserSingleton->CreateBrowserWindow(WindowSettings);
+		if (!Browser.IsValid())
+		{
+			UE_LOG(LogTemp, Error, TEXT("WebBrowserWindow is invalid! Can't display Quixel Bridge tab contents."));
+			return SAssignNew(LocalBrowserDock, SDockTab)
+				.TabRole(ETabRole::NomadTab)
+				[
+					SNew(SVerticalBox)
+						+ SVerticalBox::Slot()
+						.AutoHeight()
+						.Padding(FMargin(0, 32, 0, 0))
+						.HAlign(HAlign_Center)
+						[
+							SNew(SImage)
+								.Image(FAppStyle::Get().GetBrush("Icons.Warning"))
+						]
+						+ SVerticalBox::Slot()
+						.AutoHeight()
+						.HAlign(HAlign_Center)
+						[
+							SNew(STextBlock)
+								.Text(FText::FromString("Unable to display the Quixel Bridge Browser"))
+								.Justification(ETextJustify::Center)
+						]
+				];
+		}
 		PluginWebBrowser = SAssignNew(WebBrowserWidget, SWebBrowser, Browser)
 			.ShowAddressBar(false)
 			.ShowControls(false);
@@ -315,10 +340,28 @@ TSharedRef<SDockTab> FBridgeUIManagerImpl::CreateBridgeTab(const FSpawnTabArgs&
 	}
 	else
 	{
-		FMessageDialog::Open(EAppMsgType::Ok, FText::FromString(TEXT("Quixel Bridge requires the “Web Browser” plugin, which is disabled. Go to Edit > Plugins and search for “Web Browser” to enable it.")), FText::FromString(TEXT("Enable Web Browser Plugin")));
-
+		UE_LOG(LogTemp, Error, TEXT("WebBrowserModule isn't available! Can't display Quixel Bridge tab contents."));
 		return SAssignNew(LocalBrowserDock, SDockTab)
-			.TabRole(ETabRole::NomadTab);
+			.TabRole(ETabRole::NomadTab)
+			[
+				SNew(SVerticalBox)
+					+ SVerticalBox::Slot()
+					.AutoHeight()
+					.Padding(FMargin(0, 32, 0, 0))
+					.HAlign(HAlign_Center)
+					[
+						SNew(SImage)
+							.Image(FAppStyle::Get().GetBrush("Icons.Warning"))
+					]
+					+ SVerticalBox::Slot()
+					.AutoHeight()
+					.HAlign(HAlign_Center)
+					[
+						SNew(STextBlock)
+							.Text(FText::FromString("Unable to display the Quixel Bridge Browser"))
+							.Justification(ETextJustify::Center)
+					]
+			];
 	}
 #endif
 
diff --git a/Engine/Plugins/Fab/Source/Fab/Private/FabBrowser.cpp b/Engine/Plugins/Fab/Source/Fab/Private/FabBrowser.cpp
--- a/Engine/Plugins/Fab/Source/Fab/Private/FabBrowser.cpp
+++ b/Engine/Plugins/Fab/Source/Fab/Private/FabBrowser.cpp
@@ -29,7 +29,6 @@
 #include "Math/Vector2D.h"
 
 #include "Misc/FileHelper.h"
-#include "Misc/MessageDialog.h"
 
 #include "Styling/CoreStyle.h"
 #include "Styling/SlateStyle.h"
@@ -284,8 +283,28 @@ TSharedRef<SDockTab> FFabBrowser::OpenTab(const FSpawnTabArgs&)
 	IWebBrowserModule& WebBrowserModule = IWebBrowserModule::Get();
 	if (!IWebBrowserModule::IsAvailable() || !WebBrowserModule.IsWebModuleAvailable())
 	{
-		FMessageDialog::Open(EAppMsgType::Ok, FText::FromString("Failed to load the plugin. Please enable Web WebBrowserWindow in the plugin manager to use Emporium."));
-		return SNew(SDockTab).TabRole(ETabRole::NomadTab);
+		FAB_LOG_ERROR("WebBrowserModule isn't available! Can't display Fab tab contents.");
+		return SNew(SDockTab)
+			.TabRole(ETabRole::NomadTab)
+			[
+				SNew(SVerticalBox)
+					+ SVerticalBox::Slot()
+					.AutoHeight()
+					.Padding(FMargin(0, 32, 0, 0))
+					.HAlign(HAlign_Center)
+					[
+						SNew(SImage)
+							.Image(FAppStyle::Get().GetBrush("Icons.Warning"))
+					]
+					+ SVerticalBox::Slot()
+					.AutoHeight()
+					.HAlign(HAlign_Center)
+					[
+						SNew(STextBlock)
+							.Text(FText::FromString("Unable to display the Fab Browser"))
+							.Justification(ETextJustify::Center)
+					]
+			];
 	}
 
 	FCreateBrowserWindowSettings WindowSettings;
@@ -303,6 +322,32 @@ TSharedRef<SDockTab> FFabBrowser::OpenTab(const FSpawnTabArgs&)
 	WebBrowserSingleton->SetDevToolsShortcutEnabled(true);
 
 	WebBrowserWindow = WebBrowserSingleton->CreateBrowserWindow(WindowSettings);
+	if (!WebBrowserWindow.IsValid())
+	{
+		FAB_LOG_ERROR("WebBrowserWindow is invalid! Can't display Fab tab contents.");
+		return SNew(SDockTab)
+			.TabRole(ETabRole::NomadTab)
+			[
+				SNew(SVerticalBox)
+					+ SVerticalBox::Slot()
+					.AutoHeight()
+					.Padding(FMargin(0, 32, 0, 0))
+					.HAlign(HAlign_Center)
+					[
+						SNew(SImage)
+							.Image(FAppStyle::Get().GetBrush("Icons.Warning"))
+					]
+					+ SVerticalBox::Slot()
+					.AutoHeight()
+					.HAlign(HAlign_Center)
+					[
+						SNew(STextBlock)
+							.Text(FText::FromString("Unable to display the Fab Browser"))
+							.Justification(ETextJustify::Center)
+					]
+			];
+	}
+
 	WebBrowserWindow->OnUnhandledKeyUp().BindLambda([](const FKeyEvent&) { return true; });
 	WebBrowserWindow->OnUnhandledKeyDown().BindLambda([](const FKeyEvent&) { return true; });
 	WebBrowserWindow->OnUrlChanged().AddLambda(
